Chaos : Caching
- Moved the cache adaptor initialization and registration to the plugin startup callbacks.
#okforgethub public, internal

diff --git a/Engine/Plugins/Experimental/ChaosCaching/Source/ChaosCaching/Private/Chaos/Adapters/GeometryCollectionComponentCacheAdapter.cpp b/Engine/Plugins/Experimental/ChaosCaching/Source/ChaosCaching/Private/Chaos/Adapters/GeometryCollectionComponentCacheAdapter.cpp
--- a/Engine/Plugins/Experimental/ChaosCaching/Source/ChaosCaching/Private/Chaos/Adapters/GeometryCollectionComponentCacheAdapter.cpp
+++ b/Engine/Plugins/Experimental/ChaosCaching/Source/ChaosCaching/Private/Chaos/Adapters/GeometryCollectionComponentCacheAdapter.cpp
@@ -385,5 +385,3 @@ namespace Chaos
 	}
 
 }    // namespace Chaos
-
-Chaos::TAutoRegisterCacheAdapter<Chaos::FGeometryCollectionCacheAdapter> GeometryCollectionAdapter;
diff --git a/Engine/Plugins/Experimental/ChaosCaching/Source/ChaosCaching/Private/Chaos/Adapters/StaticMeshComponentCacheAdapter.cpp b/Engine/Plugins/Experimental/ChaosCaching/Source/ChaosCaching/Private/Chaos/Adapters/StaticMeshComponentCacheAdapter.cpp
--- a/Engine/Plugins/Experimental/ChaosCaching/Source/ChaosCaching/Private/Chaos/Adapters/StaticMeshComponentCacheAdapter.cpp
+++ b/Engine/Plugins/Experimental/ChaosCaching/Source/ChaosCaching/Private/Chaos/Adapters/StaticMeshComponentCacheAdapter.cpp
@@ -159,5 +159,3 @@ namespace Chaos
 	}
 
 }    // namespace Chaos
-
-Chaos::TAutoRegisterCacheAdapter<Chaos::FStaticMeshCacheAdapter> StaticMeshAdapter;
diff --git a/Engine/Plugins/Experimental/ChaosCaching/Source/ChaosCaching/Private/Chaos/ChaosCachingPlugin.cpp b/Engine/Plugins/Experimental/ChaosCaching/Source/ChaosCaching/Private/Chaos/ChaosCachingPlugin.cpp
--- a/Engine/Plugins/Experimental/ChaosCaching/Source/ChaosCaching/Private/Chaos/ChaosCachingPlugin.cpp
+++ b/Engine/Plugins/Experimental/ChaosCaching/Source/ChaosCaching/Private/Chaos/ChaosCachingPlugin.cpp
@@ -1,6 +1,9 @@
 // Copyright Epic Games, Inc. All Rights Reserved.
 
 #include "Chaos/ChaosCachingPlugin.h"
+
+#include "Chaos/Adapters/GeometryCollectionComponentCacheAdapter.h"
+#include "Chaos/Adapters/StaticMeshComponentCacheAdapter.h"
 #include "CoreMinimal.h"
 #include "Modules/ModuleManager.h"
 
@@ -8,6 +11,20 @@ IMPLEMENT_MODULE(IChaosCachingPlugin, ChaosCaching)
 
 DEFINE_LOG_CATEGORY(LogChaosCache)
 
-void IChaosCachingPlugin::StartupModule() {}
+void IChaosCachingPlugin::StartupModule() 
+{
+	GeometryCollectionAdapter = MakeUnique<Chaos::FGeometryCollectionCacheAdapter>();
+	StaticMeshAdapter = MakeUnique<Chaos::FStaticMeshCacheAdapter>();
+
+	RegisterAdapter(GeometryCollectionAdapter.Get());
+	RegisterAdapter(StaticMeshAdapter.Get());
+}
+
+void IChaosCachingPlugin::ShutdownModule() 
+{
+	UnregisterAdapter(StaticMeshAdapter.Get());
+	UnregisterAdapter(GeometryCollectionAdapter.Get());
 
-void IChaosCachingPlugin::ShutdownModule() {}
+	StaticMeshAdapter = nullptr;
+	GeometryCollectionAdapter = nullptr;
+}
diff --git a/Engine/Plugins/Experimental/ChaosCaching/Source/ChaosCaching/Public/Chaos/ChaosCachingPlugin.h b/Engine/Plugins/Experimental/ChaosCaching/Source/ChaosCaching/Public/Chaos/ChaosCachingPlugin.h
--- a/Engine/Plugins/Experimental/ChaosCaching/Source/ChaosCaching/Public/Chaos/ChaosCachingPlugin.h
+++ b/Engine/Plugins/Experimental/ChaosCaching/Source/ChaosCaching/Public/Chaos/ChaosCachingPlugin.h
@@ -7,6 +7,11 @@
 #include "Modules/ModuleManager.h"
 #include "Logging/LogMacros.h"
 
+namespace Chaos {
+	class FGeometryCollectionCacheAdapter;
+	class FStaticMeshCacheAdapter;
+}
+
 DECLARE_LOG_CATEGORY_EXTERN(LogChaosCache, Verbose, All);
 
 class IChaosCachingPlugin : public IModuleInterface
@@ -37,4 +42,7 @@ public:
 	}
 
 private:
+	TUniquePtr<Chaos::FGeometryCollectionCacheAdapter> GeometryCollectionAdapter;
+	TUniquePtr<Chaos::FStaticMeshCacheAdapter> StaticMeshAdapter;
+
 };
