GeometryCollectionCacheAdapter
Chaos EventHandler
-- GeometryCollectionCacheHandler captures breaking events by implementing a cache handler and creating cache events for relevent events captured during record.
-- GeometryCollectionCacheHandler emits events into the EventHandler on playback.
-- checks for null particles (usually due to non-simulating)

diff --git a/Engine/Plugins/Experimental/ChaosCaching/Source/ChaosCaching/Private/Chaos/Adapters/GeometryCollectionComponentCacheAdapter.cpp b/Engine/Plugins/Experimental/ChaosCaching/Source/ChaosCaching/Private/Chaos/Adapters/GeometryCollectionComponentCacheAdapter.cpp
--- a/Engine/Plugins/Experimental/ChaosCaching/Source/ChaosCaching/Private/Chaos/Adapters/GeometryCollectionComponentCacheAdapter.cpp
+++ b/Engine/Plugins/Experimental/ChaosCaching/Source/ChaosCaching/Private/Chaos/Adapters/GeometryCollectionComponentCacheAdapter.cpp
@@ -10,6 +10,7 @@
 #include "Chaos/ChaosSolverActor.h"
 
 FName FEnableStateEvent::EventName("GC_Enable");
+FName FBreakingEvent::EventName("GC_Breaking");
 
 namespace Chaos
 {
@@ -120,6 +121,26 @@ namespace Chaos
 			}
 		}
 
+		if (BreakingDataArray && ProxyBreakingDataIndices)
+		{
+			for (int32 Index : *ProxyBreakingDataIndices)
+			{
+				if (BreakingDataArray->IsValidIndex(Index))
+				{
+					const TBreakingData<float, 3>& BreakingData = (*BreakingDataArray)[Index];
+					if (const TPBDRigidParticleHandle<FReal, 3>*Rigid = BreakingData.Particle->CastToRigidParticle())
+					{
+						int32 TransformIndex = Proxy->GetTransformGroupIndexFromHandle(Rigid);
+						if (TransformIndex > INDEX_NONE)
+						{
+							OutFrame.PushEvent(FBreakingEvent::EventName, InTime, FBreakingEvent(TransformIndex, BreakingData));
+						}
+					}
+				}
+				
+			}
+		}
+
 		// Never going to change again till freed after writing to the cache so free up the extra space we reserved
 		OutFrame.PendingParticleData.Shrink();
 	}
@@ -163,6 +184,7 @@ namespace Chaos
 
 		const int32                      NumEventTracks = EvaluatedResult.Events.Num();
 		const TArray<FCacheEventHandle>* EnableEvents   = EvaluatedResult.Events.Find(FEnableStateEvent::EventName);
+		const TArray<FCacheEventHandle>* BreakingEvents = EvaluatedResult.Events.Find(FBreakingEvent::EventName);
 
 		if(EnableEvents)
 		{
@@ -175,26 +197,28 @@ namespace Chaos
 					{
 						Chaos::TPBDRigidClusteredParticleHandle<float, 3>* ChildParticle = Particles[Event->Index];
 						
-						if(ChildParticle->ObjectState() != EObjectStateType::Kinematic)
+						if (ChildParticle)
 						{
-							// If a field or other external actor set the particle to static or dynamic we no longer apply the cache
-							continue;
-						}
+							if (ChildParticle->ObjectState() != EObjectStateType::Kinematic)
+							{
+								// If a field or other external actor set the particle to static or dynamic we no longer apply the cache
+								continue;
+							}
 
-						if(FRigidParticle* ClusterParent = ChildParticle->ClusterIds().Id)
-						{
-							if(FClusterParticle* Parent = ClusterParent->CastToClustered())
+							if (FRigidParticle* ClusterParent = ChildParticle->ClusterIds().Id)
 							{
-								TArray<FRigidParticle*>& Cluster = NewClusters.FindOrAdd(Parent);
-								Cluster.Add(ChildParticle);
+								if (FClusterParticle* Parent = ClusterParent->CastToClustered())
+								{
+									TArray<FRigidParticle*>& Cluster = NewClusters.FindOrAdd(Parent);
+									Cluster.Add(ChildParticle);
+								}
+							}
+							else
+							{
+								// This is a cluster parent
+								ChildParticle->SetDisabled(!Event->bEnable);
 							}
 						}
-						else
-						{
-							// This is a cluster parent
-							ChildParticle->SetDisabled(!Event->bEnable);
-						}
-						
 					}
 				}
 			}
@@ -231,6 +255,54 @@ namespace Chaos
 			}
 		}
 
+		if (BreakingEvents)
+		{
+			const FSolverBreakingEventFilter* SolverBreakingEventFilter = Solver->GetEventFilters()->GetBreakingFilter();
+
+			for (const FCacheEventHandle& Handle : *BreakingEvents)
+			{
+				if (FBreakingEvent* Event = Handle.Get<FBreakingEvent>())
+				{
+					if (Particles.IsValidIndex(Event->Index))
+					{
+						Chaos::TPBDRigidClusteredParticleHandle<float, 3>* Particle = Particles[Event->Index];
+
+						if (Particle)
+						{
+							if (Particle->ObjectState() != EObjectStateType::Kinematic)
+							{
+								// If a field or other external actor set the particle to static or dynamic we no longer apply the cache
+								continue;
+							}
+
+							TBreakingData<float, 3> CachedBreak;
+							CachedBreak.Particle = Particle;
+							CachedBreak.ParticleProxy = Proxy;
+							CachedBreak.Location = Event->Location;
+							CachedBreak.Velocity = Event->Velocity;
+							CachedBreak.AngularVelocity = Event->AngularVelocity;
+							CachedBreak.Mass = Event->Mass;
+							CachedBreak.BoundingBox = TAABB<float, 3>(Event->BoundingBoxMin, Event->BoundingBoxMax);
+
+							if (!SolverBreakingEventFilter->Enabled() || SolverBreakingEventFilter->Pass(CachedBreak))
+							{
+								float TimeStamp = Solver->GetSolverTime();
+								Solver->GetEventManager()->AddEvent<FBreakingEventData>(EEventType::Breaking, [&CachedBreak, TimeStamp](FBreakingEventData& BreakingEventData)
+									{
+										if (BreakingEventData.BreakingData.TimeCreated != TimeStamp)
+										{
+											BreakingEventData.BreakingData.AllBreakingsArray.Reset();
+											BreakingEventData.BreakingData.TimeCreated = TimeStamp;
+										}
+										BreakingEventData.BreakingData.AllBreakingsArray.Add(CachedBreak);
+									});
+							}
+						}
+					}
+				}
+			}
+		}
+
 		const int32 NumTransforms = EvaluatedResult.Transform.Num();
 		for(int32 Index = 0; Index < NumTransforms; ++Index)
 		{
@@ -346,10 +418,14 @@ namespace Chaos
 		return nullptr;
 	}
 
-	bool FGeometryCollectionCacheAdapter::InitializeForRecord(UPrimitiveComponent* InComponent, UChaosCache* InCache) const
+	bool FGeometryCollectionCacheAdapter::InitializeForRecord(UPrimitiveComponent* InComponent, UChaosCache* InCache)
 	{
 		UGeometryCollectionComponent*    Comp     = CastChecked<UGeometryCollectionComponent>(InComponent);
 		FGeometryCollectionPhysicsProxy* Proxy    = Comp->GetPhysicsProxy();
+		
+		ProxyKey = Proxy;
+		BreakingDataArray = nullptr;
+		ProxyBreakingDataIndices = nullptr;
 
 		if(!Proxy)
 		{
@@ -364,7 +440,15 @@ namespace Chaos
 		}
 
 		// We need breaking data to record cluster breaking information into the cache
-		Solver->GetEvolution()->GetRigidClustering().SetGenerateClusterBreaking(true);
+		Solver->SetGenerateBreakingData(true);
+
+
+		// Initialize event handlers to capture Breaking, Collision and Trailing events.
+		Chaos::FEventManager* EventManager = Solver->GetEventManager();
+		if (EventManager)
+		{
+			EventManager->RegisterHandler<Chaos::FBreakingEventData>(Chaos::EEventType::Breaking, const_cast<FGeometryCollectionCacheAdapter*>(this), &FGeometryCollectionCacheAdapter::HandleBreakingEvents);
+		}
 
 		return true;
 	}
@@ -376,7 +460,7 @@ namespace Chaos
 
 		FGeometryDynamicCollection& Collection = Proxy->GetPhysicsCollection();
 
-		for(int32& State : Collection.DynamicState)
+		for (int32& State : Collection.DynamicState)
 		{
 			State = (int32)EObjectStateTypeEnum::Chaos_Object_Kinematic;
 		}
@@ -384,4 +468,16 @@ namespace Chaos
 		return true;
 	}
 
+	void FGeometryCollectionCacheAdapter::HandleBreakingEvents(const Chaos::FBreakingEventData& Event)
+	{
+		if (ProxyKey)
+		{
+			BreakingDataArray = &Event.BreakingData.AllBreakingsArray;
+			if (Event.PhysicsProxyToBreakingIndices.PhysicsProxyToIndicesMap.Contains(ProxyKey))
+			{
+				ProxyBreakingDataIndices = &Event.PhysicsProxyToBreakingIndices.PhysicsProxyToIndicesMap[ProxyKey];
+			}
+		}
+	}
+
 }    // namespace Chaos
diff --git a/Engine/Plugins/Experimental/ChaosCaching/Source/ChaosCaching/Private/Chaos/Adapters/StaticMeshComponentCacheAdapter.cpp b/Engine/Plugins/Experimental/ChaosCaching/Source/ChaosCaching/Private/Chaos/Adapters/StaticMeshComponentCacheAdapter.cpp
--- a/Engine/Plugins/Experimental/ChaosCaching/Source/ChaosCaching/Private/Chaos/Adapters/StaticMeshComponentCacheAdapter.cpp
+++ b/Engine/Plugins/Experimental/ChaosCaching/Source/ChaosCaching/Private/Chaos/Adapters/StaticMeshComponentCacheAdapter.cpp
@@ -140,7 +140,7 @@ namespace Chaos
 		return nullptr;
 	}
 
-	bool FStaticMeshCacheAdapter::InitializeForRecord(UPrimitiveComponent* InComponent, UChaosCache* InCache) const
+	bool FStaticMeshCacheAdapter::InitializeForRecord(UPrimitiveComponent* InComponent, UChaosCache* InCache)
 	{
 		return true;
 	}
diff --git a/Engine/Plugins/Experimental/ChaosCaching/Source/ChaosCaching/Public/Chaos/Adapters/CacheAdapter.h b/Engine/Plugins/Experimental/ChaosCaching/Source/ChaosCaching/Public/Chaos/Adapters/CacheAdapter.h
--- a/Engine/Plugins/Experimental/ChaosCaching/Source/ChaosCaching/Public/Chaos/Adapters/CacheAdapter.h
+++ b/Engine/Plugins/Experimental/ChaosCaching/Source/ChaosCaching/Public/Chaos/Adapters/CacheAdapter.h
@@ -118,7 +118,7 @@ namespace Chaos
 		 * @param InComponent Target component to initialize
 		 * @param InCache Target cache to initialize
 		 */
-		virtual bool InitializeForRecord(UPrimitiveComponent* InComponent, UChaosCache* InCache) const = 0;
+		virtual bool InitializeForRecord(UPrimitiveComponent* InComponent, UChaosCache* InCache) = 0;
 
 		/**
 		 * Called from the game thread to initialize a component and cache ready to playback a cache
diff --git a/Engine/Plugins/Experimental/ChaosCaching/Source/ChaosCaching/Public/Chaos/Adapters/GeometryCollectionComponentCacheAdapter.h b/Engine/Plugins/Experimental/ChaosCaching/Source/ChaosCaching/Public/Chaos/Adapters/GeometryCollectionComponentCacheAdapter.h
--- a/Engine/Plugins/Experimental/ChaosCaching/Source/ChaosCaching/Public/Chaos/Adapters/GeometryCollectionComponentCacheAdapter.h
+++ b/Engine/Plugins/Experimental/ChaosCaching/Source/ChaosCaching/Public/Chaos/Adapters/GeometryCollectionComponentCacheAdapter.h
@@ -4,6 +4,7 @@
 
 #include "CacheAdapter.h"
 #include "Chaos/CacheEvents.h"
+#include "EventsData.h"
 #include "GeometryCollectionComponentCacheAdapter.generated.h"
 
 USTRUCT()
@@ -32,6 +33,58 @@ struct FEnableStateEvent : public FCacheEventBase
 	bool bEnable;
 };
 
+USTRUCT()
+struct FBreakingEvent : public FCacheEventBase
+{
+	GENERATED_BODY()
+
+	static FName EventName;
+
+	FBreakingEvent()
+		: Index(INDEX_NONE)
+		, Location(0.0f, 0.0f, 0.0f)
+		, Velocity(0.0f, 0.0f, 0.0f)
+		, AngularVelocity(0.0f, 0.0f, 0.0f)
+		, Mass(1.0f)
+		, BoundingBoxMin(0.0f, 0.0f, 0.0f)
+		, BoundingBoxMax(0.0f, 0.0f, 0.0f)
+	{
+	}
+
+	FBreakingEvent(int32 InIndex, const Chaos::TBreakingData<float, 3>& InData)
+		: Index(InIndex)
+		, Location(InData.Location)
+		, Velocity(InData.Velocity)
+		, AngularVelocity(InData.AngularVelocity)
+		, Mass(InData.Mass)
+		, BoundingBoxMin(InData.BoundingBox.Min())
+		, BoundingBoxMax(InData.BoundingBox.Max())
+	{
+	}
+
+	UPROPERTY()
+	int32 Index;
+
+	UPROPERTY()
+	FVector Location;
+	
+	UPROPERTY()
+	FVector Velocity;
+
+	UPROPERTY()
+	FVector AngularVelocity;
+	
+	UPROPERTY()
+	float Mass;
+
+	UPROPERTY()
+	FVector BoundingBoxMin;
+
+	UPROPERTY()
+	FVector BoundingBoxMax;
+};
+
+
 class UChaosCache;
 class UPrimitiveComponent;
 
@@ -40,6 +93,12 @@ namespace Chaos
 	class FGeometryCollectionCacheAdapter : public FComponentCacheAdapter
 	{
 	public:
+		FGeometryCollectionCacheAdapter()
+			: ProxyKey(nullptr)
+			, BreakingDataArray(nullptr)
+			, ProxyBreakingDataIndices(nullptr)
+		{}
+
 		virtual ~FGeometryCollectionCacheAdapter() = default;
 
 		// FComponentCacheAdapter interface
@@ -49,7 +108,7 @@ namespace Chaos
 		FGuid                  GetGuid() const override;
 		Chaos::FPhysicsSolver* GetComponentSolver(UPrimitiveComponent* InComponent) const override;
 		bool                   ValidForPlayback(UPrimitiveComponent* InComponent, UChaosCache* InCache) const override;
-		bool                   InitializeForRecord(UPrimitiveComponent* InComponent, UChaosCache* InCache) const override;
+		bool                   InitializeForRecord(UPrimitiveComponent* InComponent, UChaosCache* InCache) override;
 		bool                   InitializeForPlayback(UPrimitiveComponent* InComponent, UChaosCache* InCache) const override;
 		void                   Record_PostSolve(UPrimitiveComponent* InComp, const FTransform& InRootTransform, FPendingFrameWrite& OutFrame, Chaos::FReal InTime) const override;
 		void                   Playback_PreSolve(UPrimitiveComponent*                               InComponent,
@@ -60,6 +119,13 @@ namespace Chaos
 
 		// End FComponentCacheAdapter interface
 
+	protected:
+		void HandleBreakingEvents(const Chaos::FBreakingEventData& Event);
+	
 	private:
+		const IPhysicsProxyBase* ProxyKey;
+		const Chaos::FBreakingDataArray* BreakingDataArray;
+		const TArray<int32>* ProxyBreakingDataIndices;
+
 	};
 }    // namespace Chaos
diff --git a/Engine/Plugins/Experimental/ChaosCaching/Source/ChaosCaching/Public/Chaos/Adapters/StaticMeshComponentCacheAdapter.h b/Engine/Plugins/Experimental/ChaosCaching/Source/ChaosCaching/Public/Chaos/Adapters/StaticMeshComponentCacheAdapter.h
--- a/Engine/Plugins/Experimental/ChaosCaching/Source/ChaosCaching/Public/Chaos/Adapters/StaticMeshComponentCacheAdapter.h
+++ b/Engine/Plugins/Experimental/ChaosCaching/Source/ChaosCaching/Public/Chaos/Adapters/StaticMeshComponentCacheAdapter.h
@@ -17,7 +17,7 @@ namespace Chaos
 		FGuid                  GetGuid() const override;
 		bool                   ValidForPlayback(UPrimitiveComponent* InComponent, UChaosCache* InCache) const override;
 		Chaos::FPhysicsSolver* GetComponentSolver(UPrimitiveComponent* InComponent) const override;
-		bool                   InitializeForRecord(UPrimitiveComponent* InComponent, UChaosCache* InCache) const override;
+		bool                   InitializeForRecord(UPrimitiveComponent* InComponent, UChaosCache* InCache) override;
 		bool                   InitializeForPlayback(UPrimitiveComponent* InComponent, UChaosCache* InCache) const override;
 		void                   Record_PostSolve(UPrimitiveComponent* InComp, const FTransform& InRootTransform, FPendingFrameWrite& OutFrame, Chaos::FReal InTime) const override;
 		void                   Playback_PreSolve(UPrimitiveComponent*                               InComponent,
diff --git a/Engine/Source/Runtime/Experimental/Chaos/Private/EventDefaults.cpp b/Engine/Source/Runtime/Experimental/Chaos/Private/EventDefaults.cpp
--- a/Engine/Source/Runtime/Experimental/Chaos/Private/EventDefaults.cpp
+++ b/Engine/Source/Runtime/Experimental/Chaos/Private/EventDefaults.cpp
@@ -246,10 +246,14 @@ namespace Chaos
 				return;
 
 			FBreakingDataArray& AllBreakingDataArray = BreakingEventData.BreakingData.AllBreakingsArray;
+			TMap<IPhysicsProxyBase*, TArray<int32>>& AllBreakingIndicesByPhysicsProxy = BreakingEventData.PhysicsProxyToBreakingIndices.PhysicsProxyToIndicesMap;
 
-			AllBreakingDataArray.Reset();
-
-			BreakingEventData.BreakingData.TimeCreated = Solver->MTime;
+			if (BreakingEventData.BreakingData.TimeCreated != Solver->MTime)
+			{
+				AllBreakingDataArray.Reset();
+				AllBreakingIndicesByPhysicsProxy.Reset();
+				BreakingEventData.BreakingData.TimeCreated = Solver->MTime;
+			}
 
 			const auto* Evolution = Solver->GetEvolution();
 			const TPBDRigidParticles<float, 3>& Particles = Evolution->GetParticles().GetDynamicParticles();
@@ -294,6 +298,9 @@ namespace Chaos
 								TBreakingData<float, 3>& BreakingDataArrayItem = AllBreakingDataArray[NewIdx];
 								BreakingDataArrayItem = BreakingData;
 
+								// Add to AllBreakingIndicesByPhysicsProxy
+								AllBreakingIndicesByPhysicsProxy.FindOrAdd(BreakingData.ParticleProxy).Add(TEventManager<Traits>::EncodeCollisionIndex(NewIdx, false));
+
 #if 0 // #todo
 								// If AllBreakingsArray[Idx].ParticleIndex is a cluster store an index for a mesh in this cluster
 								if(ClusterIdsArray[AllBreakingsArray[Idx].ParticleIndex].NumChildren > 0)
@@ -308,6 +315,7 @@ namespace Chaos
 					}
 				}
 			}
+		
 		});
 	}
 
diff --git a/Engine/Source/Runtime/Experimental/Chaos/Public/EventManager.h b/Engine/Source/Runtime/Experimental/Chaos/Public/EventManager.h
--- a/Engine/Source/Runtime/Experimental/Chaos/Public/EventManager.h
+++ b/Engine/Source/Runtime/Experimental/Chaos/Public/EventManager.h
@@ -187,11 +187,17 @@ namespace Chaos
 			InjectedFunction(Solver, *EventBuffer->AccessProducerBuffer());
 		}
 
+		
 		virtual void DestroyStaleEvents(TFunction<void(PayloadType & EventDataInOut)> InFunction)
 		{
 			InFunction(*EventBuffer->AccessProducerBuffer());
 		}
 
+		virtual void AddEvent(TFunction<void(PayloadType& EventDataInOut)> InFunction)
+		{
+			InFunction(*EventBuffer->AccessProducerBuffer());
+		}
+
 		/**
 		 * Flips the buffer if the buffer type is double or triple
 		 */
@@ -288,7 +294,10 @@ namespace Chaos
 		{
 			ContainerLock.ReadLock();
 
-			((TEventContainer<PayloadType, Traits>*)(EventContainers[FEventID(EventType)]))->DestroyStaleEvents(InFunction);
+			if (TEventContainer<PayloadType, Traits>* EventContainer = StaticCast<TEventContainer<PayloadType, Traits>*>(EventContainers[FEventID(EventType)]))
+			{
+				EventContainer->DestroyStaleEvents(InFunction);
+			}
 			ContainerLock.ReadUnlock();
 		}
 
@@ -335,6 +344,18 @@ namespace Chaos
 		/** Returns decoded collision index. */
 		static int32 DecodeCollisionIndex(int32 EncodedCollisionIdx, bool& bSwapOrder);
 
+
+		template<typename PayloadType>
+		void AddEvent(const EEventType& EventType, TFunction<void(PayloadType& EventData)> InFunction)
+		{
+			ContainerLock.ReadLock();
+			if (TEventContainer<PayloadType, Traits>* EventContainer = StaticCast<TEventContainer<PayloadType, Traits>*>(EventContainers[FEventID(EventType)]))
+			{
+				EventContainer->AddEvent(InFunction);
+			}
+			ContainerLock.ReadUnlock();
+		}
+
 	private:
 
 		void InternalRegisterInjector(const FEventID& EventID, const TEventContainerBasePtr<Traits>& Container);
diff --git a/Engine/Source/Runtime/Experimental/Chaos/Public/EventsData.h b/Engine/Source/Runtime/Experimental/Chaos/Public/EventsData.h
--- a/Engine/Source/Runtime/Experimental/Chaos/Public/EventsData.h
+++ b/Engine/Source/Runtime/Experimental/Chaos/Public/EventsData.h
@@ -6,6 +6,8 @@
 #include "Chaos/CollisionResolutionTypes.h"
 #include "Chaos/Framework/PhysicsProxy.h"
 
+class UPrimitiveComponent;
+
 namespace Chaos
 {
 	// base class for data that requires time of creation to be recorded
@@ -86,6 +88,7 @@ namespace Chaos
 		FBreakingEventData() {}
 
 		FAllBreakingData BreakingData;
+		FIndicesByPhysicsProxy PhysicsProxyToBreakingIndices;
 	};
 
 	/* Trailing */
diff --git a/Engine/Source/Runtime/Experimental/Chaos/Public/PhysicsProxy/GeometryCollectionPhysicsProxy.h b/Engine/Source/Runtime/Experimental/Chaos/Public/PhysicsProxy/GeometryCollectionPhysicsProxy.h
--- a/Engine/Source/Runtime/Experimental/Chaos/Public/PhysicsProxy/GeometryCollectionPhysicsProxy.h
+++ b/Engine/Source/Runtime/Experimental/Chaos/Public/PhysicsProxy/GeometryCollectionPhysicsProxy.h
@@ -229,6 +229,18 @@ public:
 	/* Implemented so we can construct TAccelerationStructureHandle. */
 	virtual void* GetHandleUnsafe() const override { return nullptr; }
 
+	int32 GetTransformGroupIndexFromHandle(const FParticleHandle* Handle) const
+	{
+		if (HandleToTransformGroupIndex.Contains(Handle))
+		{
+			return HandleToTransformGroupIndex[Handle];
+		}
+		else
+		{
+			return INDEX_NONE;
+		}
+	}
+
 protected:
 	/**
 	 * Build a physics thread cluster parent particle.
