Chaos Cache
-- fixed bug where placed cache spawned double observed components

diff --git a/Engine/Plugins/Experimental/ChaosCaching/Source/ChaosCaching/Private/Chaos/CacheManagerActor.cpp b/Engine/Plugins/Experimental/ChaosCaching/Source/ChaosCaching/Private/Chaos/CacheManagerActor.cpp
--- a/Engine/Plugins/Experimental/ChaosCaching/Source/ChaosCaching/Private/Chaos/CacheManagerActor.cpp
+++ b/Engine/Plugins/Experimental/ChaosCaching/Source/ChaosCaching/Private/Chaos/CacheManagerActor.cpp
@@ -541,6 +541,16 @@ FObservedComponent& AChaosCacheManager::FindOrAddObservedComponent(UPrimitiveCom
 	return Found ? *Found : AddNewObservedComponent(InComponent);
 }
 
+void AChaosCacheManager::ClearObservedComponents()
+{
+	for (FObservedComponent& ObservedComponent : ObservedComponents)
+	{
+		UPrimitiveComponent* Comp = ObservedComponent.GetComponent();
+		Comp->DestroyComponent();
+	}
+	ObservedComponents.Empty();
+}
+
 void AChaosCacheManager::TickObservedComponents(const TArray<int32>& InIndices, Chaos::FReal InDt, FTickObservedFunction InCallable)
 {
 	for(int32 Index : InIndices)
diff --git a/Engine/Plugins/Experimental/ChaosCaching/Source/ChaosCaching/Public/Chaos/CacheManagerActor.h b/Engine/Plugins/Experimental/ChaosCaching/Source/ChaosCaching/Public/Chaos/CacheManagerActor.h
--- a/Engine/Plugins/Experimental/ChaosCaching/Source/ChaosCaching/Public/Chaos/CacheManagerActor.h
+++ b/Engine/Plugins/Experimental/ChaosCaching/Source/ChaosCaching/Public/Chaos/CacheManagerActor.h
@@ -210,6 +210,7 @@ protected:
 	FObservedComponent* FindObservedComponent(UPrimitiveComponent* InComponent);
 	FObservedComponent& AddNewObservedComponent(UPrimitiveComponent* InComponent);
 	FObservedComponent& FindOrAddObservedComponent(UPrimitiveComponent* InComponent);
+	void ClearObservedComponents();
 
 private:
 	friend class UActorFactoryCacheManager; // Allows the actor factory to set up the observed list. See UActorFactoryCacheManager::PostSpawnActor
diff --git a/Engine/Plugins/Experimental/ChaosCaching/Source/ChaosCachingEditor/Private/Chaos/ActorFactoryCacheManager.cpp b/Engine/Plugins/Experimental/ChaosCaching/Source/ChaosCachingEditor/Private/Chaos/ActorFactoryCacheManager.cpp
--- a/Engine/Plugins/Experimental/ChaosCaching/Source/ChaosCachingEditor/Private/Chaos/ActorFactoryCacheManager.cpp
+++ b/Engine/Plugins/Experimental/ChaosCaching/Source/ChaosCachingEditor/Private/Chaos/ActorFactoryCacheManager.cpp
@@ -32,35 +32,42 @@ void UActorFactoryCacheManager::PostSpawnActor(UObject* Asset, AActor* NewActor)
 {
 	AChaosCacheManager*    Manager    = Cast<AChaosCacheManager>(NewActor);
 	UChaosCacheCollection* Collection = Cast<UChaosCacheCollection>(Asset);
+	
 	// The cachemanager exists now, start adding our spawnables
-	if(Manager && Collection)
+	if (Manager && Collection)
 	{
-		if(UWorld* World = Manager->GetWorld())
+		if (!Manager->bIsEditorPreviewActor)
 		{
-			Manager->CacheCollection = Collection;
-			for(UChaosCache* Cache : Collection->GetCaches())
+			if (UWorld* World = Manager->GetWorld())
 			{
-				if(!Cache)
+				// Remove any pre-existing components
+				Manager->ClearObservedComponents();
+				
+				Manager->CacheCollection = Collection;
+				for (UChaosCache* Cache : Collection->GetCaches())
 				{
-					continue;
-				}
+					if (!Cache)
+					{
+						continue;
+					}
 
-				const FCacheSpawnableTemplate& Template = Cache->GetSpawnableTemplate();
-				if(Template.DuplicatedTemplate)
-				{
-					check(Template.DuplicatedTemplate->GetClass()->IsChildOf(UPrimitiveComponent::StaticClass()));
+					const FCacheSpawnableTemplate& Template = Cache->GetSpawnableTemplate();
+					if (Template.DuplicatedTemplate)
+					{
+						check(Template.DuplicatedTemplate->GetClass()->IsChildOf(UPrimitiveComponent::StaticClass()));
 
-					UPrimitiveComponent* NewComponent = CastChecked<UPrimitiveComponent>(StaticDuplicateObject(Template.DuplicatedTemplate, Manager));
-					NewComponent->SetWorldTransform(Template.InitialTransform);
-					Manager->AddInstanceComponent(NewComponent);
-					NewComponent->RegisterComponent();
+						UPrimitiveComponent* NewComponent = CastChecked<UPrimitiveComponent>(StaticDuplicateObject(Template.DuplicatedTemplate, Manager));
+						NewComponent->SetWorldTransform(Template.InitialTransform);
+						Manager->AddInstanceComponent(NewComponent);
+						NewComponent->RegisterComponent();
 
-					FObservedComponent& Observed = Manager->AddNewObservedComponent(NewComponent);
-					// AddNewObservedComponent will have given this a unique name as if it was going to build a new cache, override to the actual cache name
-					Observed.CacheName     = Cache->GetFName();
-					Observed.CacheMode     = ECacheMode::Play;
-					Observed.StartMode     = EStartMode::Timed;
-					Observed.TimedDuration = 0.0f;
+						FObservedComponent& Observed = Manager->AddNewObservedComponent(NewComponent);
+						// AddNewObservedComponent will have given this a unique name as if it was going to build a new cache, override to the actual cache name
+						Observed.CacheName = Cache->GetFName();
+						Observed.CacheMode = ECacheMode::Play;
+						Observed.StartMode = EStartMode::Timed;
+						Observed.TimedDuration = 0.0f;
+					}
 				}
 			}
 		}
