ChaosCache
-- Initialized FCacheSpawnableTemplate struct members

diff --git a/Engine/Plugins/Experimental/ChaosCaching/Source/ChaosCaching/Public/Chaos/ChaosCache.h b/Engine/Plugins/Experimental/ChaosCaching/Source/ChaosCaching/Public/Chaos/ChaosCache.h
--- a/Engine/Plugins/Experimental/ChaosCaching/Source/ChaosCaching/Public/Chaos/ChaosCache.h
+++ b/Engine/Plugins/Experimental/ChaosCaching/Source/ChaosCaching/Public/Chaos/ChaosCache.h
@@ -70,6 +70,11 @@ struct FCacheSpawnableTemplate
 {
 	GENERATED_BODY()
 
+	FCacheSpawnableTemplate()
+		: DuplicatedTemplate(nullptr)
+		, InitialTransform(FTransform::Identity)
+	{}
+
 	UPROPERTY(VisibleAnywhere, Category = "Caching")
 	UObject* DuplicatedTemplate;
 
