Chaos Caching
-- Added buttons to CacheManager detail for setting all observed components start modes.

diff --git a/Engine/Plugins/Experimental/ChaosCaching/Source/ChaosCaching/Private/Chaos/CacheManagerActor.cpp b/Engine/Plugins/Experimental/ChaosCaching/Source/ChaosCaching/Private/Chaos/CacheManagerActor.cpp
--- a/Engine/Plugins/Experimental/ChaosCaching/Source/ChaosCaching/Private/Chaos/CacheManagerActor.cpp
+++ b/Engine/Plugins/Experimental/ChaosCaching/Source/ChaosCaching/Private/Chaos/CacheManagerActor.cpp
@@ -114,6 +114,14 @@ void AChaosCacheManager::SetAllMode(ECacheMode InMode)
 	}
 }
 
+void AChaosCacheManager::SetAllStartMode(EStartMode InMode)
+{
+	for (FObservedComponent& Observed : ObservedComponents)
+	{
+		Observed.StartMode = InMode;
+	}
+}
+
 void AChaosCacheManager::ResetAllComponentTransforms()
 {
 	if(!CacheCollection)
diff --git a/Engine/Plugins/Experimental/ChaosCaching/Source/ChaosCaching/Public/Chaos/CacheManagerActor.h b/Engine/Plugins/Experimental/ChaosCaching/Source/ChaosCaching/Public/Chaos/CacheManagerActor.h
--- a/Engine/Plugins/Experimental/ChaosCaching/Source/ChaosCaching/Public/Chaos/CacheManagerActor.h
+++ b/Engine/Plugins/Experimental/ChaosCaching/Source/ChaosCaching/Public/Chaos/CacheManagerActor.h
@@ -147,6 +147,13 @@ public:
 	UFUNCTION(BlueprintCallable, Category = "Caching")
 	void SetAllMode(ECacheMode InMode);
 
+	/**
+	 * Sets the start mode of every observed component to the specified mode
+	 * @param InMode Mode to set
+	 */
+	UFUNCTION(BlueprintCallable, Category = "Caching")
+	void SetAllStartMode(EStartMode InMode);
+
 	/** 
 	 * Resets all components back to the world space transform they had when the cache for them was originally recorded
 	 * if one is available
diff --git a/Engine/Plugins/Experimental/ChaosCaching/Source/ChaosCachingEditor/Private/Chaos/CacheEditorCommands.cpp b/Engine/Plugins/Experimental/ChaosCaching/Source/ChaosCachingEditor/Private/Chaos/CacheEditorCommands.cpp
--- a/Engine/Plugins/Experimental/ChaosCaching/Source/ChaosCachingEditor/Private/Chaos/CacheEditorCommands.cpp
+++ b/Engine/Plugins/Experimental/ChaosCaching/Source/ChaosCachingEditor/Private/Chaos/CacheEditorCommands.cpp
@@ -9,6 +9,8 @@ void FCachingEditorCommands::RegisterCommands()
 	UI_COMMAND(CreateCacheManager, "Create Cache Manager", "Adds a cache manager to observe compatible components in the selection set.", EUserInterfaceActionType::Button, FInputChord());
 	UI_COMMAND(SetManagerRecordAll, "Set Record All", "Sets selected cache managers to record all of their observed components.", EUserInterfaceActionType::Button, FInputChord());
 	UI_COMMAND(SetManagerPlayAll, "Set Play All", "Sets selected cache managers to play all of their observed components.", EUserInterfaceActionType::Button, FInputChord());
+	UI_COMMAND(SetManagerTriggerAll, "Set Trigger All", "Sets selected cache managers to set all of their observed components to triggered.", EUserInterfaceActionType::Button, FInputChord());
+	UI_COMMAND(SetManagerTimedAll, "Set Timed All", "Sets selected cache managers to set all of their observed components to timed.", EUserInterfaceActionType::Button, FInputChord());
 }
 
 #undef LOCTEXT_NAMESPACE
diff --git a/Engine/Plugins/Experimental/ChaosCaching/Source/ChaosCachingEditor/Private/Chaos/CacheManagerCustomization.cpp b/Engine/Plugins/Experimental/ChaosCaching/Source/ChaosCachingEditor/Private/Chaos/CacheManagerCustomization.cpp
--- a/Engine/Plugins/Experimental/ChaosCaching/Source/ChaosCachingEditor/Private/Chaos/CacheManagerCustomization.cpp
+++ b/Engine/Plugins/Experimental/ChaosCaching/Source/ChaosCachingEditor/Private/Chaos/CacheManagerCustomization.cpp
@@ -25,6 +25,19 @@ FReply OnClickSetAllButton(TArray<AChaosCacheManager*> Managers, ECacheMode NewM
 	return FReply::Handled();
 }
 
+FReply OnClickSetAllStartButton(TArray<AChaosCacheManager*> Managers, EStartMode NewMode)
+{
+	for (AChaosCacheManager* Manager : Managers)
+	{
+		if (Manager)
+		{
+			Manager->SetAllStartMode(NewMode);
+		}
+	}
+
+	return FReply::Handled();
+}
+
 FReply OnClickResetTransforms(TArray<AChaosCacheManager*> Managers)
 {
 	for(AChaosCacheManager* Manager : Managers)
@@ -113,6 +126,44 @@ void FCacheManagerDetails::CustomizeDetails(IDetailLayoutBuilder& DetailBuilder)
 		]
 	];
 
+	FDetailWidgetRow& SetTriggerRow = CachingCategory.AddCustomRow(FText::GetEmpty());
+	SetTriggerRow.NameContent()
+	[
+		SNew(STextBlock)
+		.Font(DetailBuilder.GetDetailFont())
+		.Text(LOCTEXT("SetTriggerLabel", "Set Triggers"))
+	];
+
+	SetTriggerRow.ValueContent()
+	.MinDesiredWidth(300.0f)
+	[
+		SNew(SHorizontalBox)
+		+ SHorizontalBox::Slot()
+		.Padding(0.0f, 0.0f, 3.0f, 0.0f)
+		[
+			SNew(SButton)
+			.HAlign(HAlign_Center)
+			.VAlign(VAlign_Center)
+			.OnClicked_Static(&OnClickSetAllStartButton, CacheManagersInSelection, EStartMode::Triggered)
+			[
+				SNew(STextBlock)
+				.Text(LOCTEXT("SetAllTriggered", "Triggered"))
+			]
+		]
+		+ SHorizontalBox::Slot()
+		.Padding(0.0f, 0.0f, 3.0f, 0.0f)
+		[
+			SNew(SButton)
+			.HAlign(HAlign_Center)
+			.VAlign(VAlign_Center)
+			.OnClicked_Static(&OnClickSetAllStartButton, CacheManagersInSelection, EStartMode::Timed)
+			[
+				SNew(STextBlock)
+				.Text(LOCTEXT("SetAllTimed", "Timed"))
+			]
+		]
+	];
+
 	FDetailWidgetRow& ResetPositionsRow = CachingCategory.AddCustomRow(FText::GetEmpty());
 
 	ResetPositionsRow.ValueContent()
diff --git a/Engine/Plugins/Experimental/ChaosCaching/Source/ChaosCachingEditor/Public/Chaos/CacheEditorCommands.h b/Engine/Plugins/Experimental/ChaosCaching/Source/ChaosCachingEditor/Public/Chaos/CacheEditorCommands.h
--- a/Engine/Plugins/Experimental/ChaosCaching/Source/ChaosCachingEditor/Public/Chaos/CacheEditorCommands.h
+++ b/Engine/Plugins/Experimental/ChaosCaching/Source/ChaosCachingEditor/Public/Chaos/CacheEditorCommands.h
@@ -22,4 +22,6 @@ public:
 	TSharedPtr<FUICommandInfo> CreateCacheManager;
 	TSharedPtr<FUICommandInfo> SetManagerRecordAll;
 	TSharedPtr<FUICommandInfo> SetManagerPlayAll;
+	TSharedPtr<FUICommandInfo> SetManagerTriggerAll;
+	TSharedPtr<FUICommandInfo> SetManagerTimedAll;
 };
