ChaosCaching
-- Correct flags are passed to new component on placed cache so that compoennts save correctly.

diff --git a/Engine/Plugins/Experimental/ChaosCaching/Source/ChaosCachingEditor/Private/Chaos/ActorFactoryCacheManager.cpp b/Engine/Plugins/Experimental/ChaosCaching/Source/ChaosCachingEditor/Private/Chaos/ActorFactoryCacheManager.cpp
--- a/Engine/Plugins/Experimental/ChaosCaching/Source/ChaosCachingEditor/Private/Chaos/ActorFactoryCacheManager.cpp
+++ b/Engine/Plugins/Experimental/ChaosCaching/Source/ChaosCachingEditor/Private/Chaos/ActorFactoryCacheManager.cpp
@@ -56,7 +56,13 @@ void UActorFactoryCacheManager::PostSpawnActor(UObject* Asset, AActor* NewActor)
 					{
 						check(Template.DuplicatedTemplate->GetClass()->IsChildOf(UPrimitiveComponent::StaticClass()));
 
-						UPrimitiveComponent* NewComponent = CastChecked<UPrimitiveComponent>(StaticDuplicateObject(Template.DuplicatedTemplate, Manager));
+						FObjectDuplicationParameters Parameters(Template.DuplicatedTemplate, Manager);
+						Parameters.FlagMask &= ~RF_Transient;
+						Parameters.FlagMask &= ~RF_DefaultSubObject;
+						Parameters.FlagMask &= ~RF_Transactional;
+						
+						UPrimitiveComponent* NewComponent = CastChecked<UPrimitiveComponent>(StaticDuplicateObjectEx(Parameters));
+						NewComponent->SetupAttachment(Manager->GetRootComponent());
 						NewComponent->SetWorldTransform(Template.InitialTransform);
 						Manager->AddInstanceComponent(NewComponent);
 						NewComponent->RegisterComponent();
