Chaos Caching
-- Caches now occur relative to the recording cache manager. Playback is likewise. If a cache is spawned, playback is relative to the spawned position.

diff --git a/Engine/Plugins/Experimental/ChaosCaching/Source/ChaosCaching/Private/Chaos/CacheManagerActor.cpp b/Engine/Plugins/Experimental/ChaosCaching/Source/ChaosCaching/Private/Chaos/CacheManagerActor.cpp
--- a/Engine/Plugins/Experimental/ChaosCaching/Source/ChaosCaching/Private/Chaos/CacheManagerActor.cpp
+++ b/Engine/Plugins/Experimental/ChaosCaching/Source/ChaosCaching/Private/Chaos/CacheManagerActor.cpp
@@ -284,9 +284,9 @@ void AChaosCacheManager::BeginPlay()
 						{
 							SolverData->PlaybackIndices.Add(Index);
 							SolverData->PlaybackTickRecords.AddDefaulted();
-							SolverData->PlaybackTickRecords.Last().SetSpaceTransform(Comp->GetComponentToWorld());
+							SolverData->PlaybackTickRecords.Last().SetSpaceTransform(GetTransform());
 							Observed.Cache = PlayCache;
-							Observed.TickRecord.SetSpaceTransform(Comp->GetComponentToWorld());
+							Observed.TickRecord.SetSpaceTransform(GetTransform());
 							OpenPlaybackCaches.Add(TTuple<FCacheUserToken, UChaosCache*>(MoveTemp(Token), Observed.Cache));
 							CurrAdapter->InitializeForPlayback(Comp, Observed.Cache);
 						}
@@ -328,14 +328,16 @@ void AChaosCacheManager::BeginPlay()
 					FName CacheName = Observed.CacheName == NAME_None ? MakeUniqueObjectName(CacheCollection, UChaosCache::StaticClass(), "Cache") : Observed.CacheName;
 
 					UChaosCache*    RecordCache = CacheCollection->FindOrAddCache(CacheName);
-					FCacheUserToken Token       = RecordCache->BeginRecord(Observed.GetComponent(), CurrAdapter->GetGuid());
+					FCacheUserToken Token       = RecordCache->BeginRecord(Observed.GetComponent(), CurrAdapter->GetGuid(), GetTransform());
 
 					if(Token.IsOpen())
 					{
 						SolverData->RecordIndices.Add(Index);
 
 						Observed.Cache = CacheCollection->FindOrAddCache(CacheName);
-						Observed.TickRecord.SetSpaceTransform(Comp->GetComponentToWorld());
+						
+						// We'll record the observed component in Cache Manager's local space.
+						Observed.TickRecord.SetSpaceTransform(GetTransform());
 						OpenRecordCaches.Add(TTuple<FCacheUserToken, UChaosCache*>(MoveTemp(Token), Observed.Cache));
 						CurrAdapter->InitializeForRecord(Component, Observed.Cache);
 
diff --git a/Engine/Plugins/Experimental/ChaosCaching/Source/ChaosCaching/Private/Chaos/ChaosCache.cpp b/Engine/Plugins/Experimental/ChaosCaching/Source/ChaosCaching/Private/Chaos/ChaosCache.cpp
--- a/Engine/Plugins/Experimental/ChaosCaching/Source/ChaosCaching/Private/Chaos/ChaosCache.cpp
+++ b/Engine/Plugins/Experimental/ChaosCaching/Source/ChaosCaching/Private/Chaos/ChaosCache.cpp
@@ -93,7 +93,7 @@ void UChaosCache::FlushPendingFrames()
 	}
 }
 
-FCacheUserToken UChaosCache::BeginRecord(UPrimitiveComponent* InComponent, FGuid InAdapterId)
+FCacheUserToken UChaosCache::BeginRecord(UPrimitiveComponent* InComponent, FGuid InAdapterId, const FTransform& SpaceTransform)
 {
 	// First make sure we're valid to record
 	int32 OtherRecordersCount = CurrentRecordCount.fetch_add(1);
@@ -114,7 +114,7 @@ FCacheUserToken UChaosCache::BeginRecord(UPrimitiveComponent* InComponent, FGuid
 			PendingWrites.Empty();
 
 			// Initialise the spawnable template to handle the provided component
-			BuildSpawnableFromComponent(InComponent);
+			BuildSpawnableFromComponent(InComponent, SpaceTransform);
 
 			// Build a compatibility hash for the component state.
 			AdapterGuid = InAdapterId;
@@ -338,10 +338,11 @@ FCacheEvaluationResult UChaosCache::Evaluate(const FCacheEvaluationContext& InCo
 	return Result;
 }
 
-void UChaosCache::BuildSpawnableFromComponent(UPrimitiveComponent* InComponent)
+void UChaosCache::BuildSpawnableFromComponent(UPrimitiveComponent* InComponent, const FTransform& SpaceTransform)
 {
 	Spawnable.DuplicatedTemplate = StaticDuplicateObject(InComponent, this);
 	Spawnable.InitialTransform = InComponent->GetComponentToWorld();
+	Spawnable.ComponentTransform = InComponent->GetComponentToWorld() * SpaceTransform.Inverse();
 }
 
 const FCacheSpawnableTemplate& UChaosCache::GetSpawnableTemplate() const
diff --git a/Engine/Plugins/Experimental/ChaosCaching/Source/ChaosCaching/Public/Chaos/ChaosCache.h b/Engine/Plugins/Experimental/ChaosCaching/Source/ChaosCaching/Public/Chaos/ChaosCache.h
--- a/Engine/Plugins/Experimental/ChaosCaching/Source/ChaosCaching/Public/Chaos/ChaosCache.h
+++ b/Engine/Plugins/Experimental/ChaosCaching/Source/ChaosCaching/Public/Chaos/ChaosCache.h
@@ -73,6 +73,7 @@ struct FCacheSpawnableTemplate
 	FCacheSpawnableTemplate()
 		: DuplicatedTemplate(nullptr)
 		, InitialTransform(FTransform::Identity)
+		, ComponentTransform(FTransform::Identity)
 	{}
 
 	UPROPERTY(VisibleAnywhere, Category = "Caching")
@@ -80,6 +81,9 @@ struct FCacheSpawnableTemplate
 
 	UPROPERTY(VisibleAnywhere, Category = "Caching")
 	FTransform InitialTransform;
+
+	UPROPERTY(VisibleAnywhere, Category = "Caching")
+	FTransform ComponentTransform;
 };
 
 struct FPlaybackTickRecord
@@ -249,7 +253,7 @@ public:
 	 * Reset and initialize a cache to make it ready to record the specified component
 	 * @param InComponent Component to prepare the cache for
 	 */
-	FCacheUserToken BeginRecord(UPrimitiveComponent* InComponent, FGuid InAdapterId);
+	FCacheUserToken BeginRecord(UPrimitiveComponent* InComponent, FGuid InAdapterId, const FTransform& SpaceTransform);
 
 	/**
 	 * End the recording session for the cache. At this point the cache is deemed to now contain
@@ -296,7 +300,7 @@ public:
 	 * when a cache is dragged into the scene.
 	 * @param InComponent Component to build the spawnable template from
 	 */
-	void BuildSpawnableFromComponent(UPrimitiveComponent* InComponent);
+	void BuildSpawnableFromComponent(UPrimitiveComponent* InComponent, const FTransform& SpaceTransform);
 
 	/**
 	 * Read access to the spawnable template stored in the cache
diff --git a/Engine/Plugins/Experimental/ChaosCaching/Source/ChaosCachingEditor/Private/Chaos/ActorFactoryCacheManager.cpp b/Engine/Plugins/Experimental/ChaosCaching/Source/ChaosCachingEditor/Private/Chaos/ActorFactoryCacheManager.cpp
--- a/Engine/Plugins/Experimental/ChaosCaching/Source/ChaosCachingEditor/Private/Chaos/ActorFactoryCacheManager.cpp
+++ b/Engine/Plugins/Experimental/ChaosCaching/Source/ChaosCachingEditor/Private/Chaos/ActorFactoryCacheManager.cpp
@@ -63,7 +63,7 @@ void UActorFactoryCacheManager::PostSpawnActor(UObject* Asset, AActor* NewActor)
 						
 						UPrimitiveComponent* NewComponent = CastChecked<UPrimitiveComponent>(StaticDuplicateObjectEx(Parameters));
 						NewComponent->SetupAttachment(Manager->GetRootComponent());
-						NewComponent->SetWorldTransform(Template.InitialTransform);
+						NewComponent->SetRelativeTransform(Template.ComponentTransform);
 						Manager->AddInstanceComponent(NewComponent);
 						NewComponent->RegisterComponent();
 
