Geometry Collection
-- Embedded geometry tweaked to make them compatible with chaos caching.

diff --git a/Engine/Plugins/Experimental/ChaosCaching/Source/ChaosCachingEditor/Private/Chaos/ChaosCachingEditorPlugin.cpp b/Engine/Plugins/Experimental/ChaosCaching/Source/ChaosCachingEditor/Private/Chaos/ChaosCachingEditorPlugin.cpp
--- a/Engine/Plugins/Experimental/ChaosCaching/Source/ChaosCachingEditor/Private/Chaos/ChaosCachingEditorPlugin.cpp
+++ b/Engine/Plugins/Experimental/ChaosCaching/Source/ChaosCachingEditor/Private/Chaos/ChaosCachingEditorPlugin.cpp
@@ -154,7 +154,7 @@ void IChaosCachingEditorPlugin::OnCreateCacheManager()
 
 			if(UPrimitiveComponent* PrimitiveComp = Cast<UPrimitiveComponent>(Component))
 			{
-				Chaos::FComponentCacheAdapter* BestFitAdapter = Chaos::FAdapterUtil::GetBestAdapterForClass(PrimitiveComp->GetClass());
+				Chaos::FComponentCacheAdapter* BestFitAdapter = Chaos::FAdapterUtil::GetBestAdapterForClass(PrimitiveComp->GetClass(), false);
 
 				// Can't be observed
 				if(!BestFitAdapter)
diff --git a/Engine/Source/Runtime/Experimental/GeometryCollectionEngine/Private/GeometryCollection/GeometryCollectionComponent.cpp b/Engine/Source/Runtime/Experimental/GeometryCollectionEngine/Private/GeometryCollection/GeometryCollectionComponent.cpp
--- a/Engine/Source/Runtime/Experimental/GeometryCollectionEngine/Private/GeometryCollection/GeometryCollectionComponent.cpp
+++ b/Engine/Source/Runtime/Experimental/GeometryCollectionEngine/Private/GeometryCollection/GeometryCollectionComponent.cpp
@@ -2961,8 +2961,7 @@ void UGeometryCollectionComponent::InitializeEmbeddedGeometry()
 					ISMC->SetCollisionProfileName(UCollisionProfile::NoCollision_ProfileName);
 					ISMC->SetCastShadow(false);
 					ISMC->SetMobility(EComponentMobility::Stationary);
-
-					ISMC->SetupAttachment(ActorOwner->GetRootComponent());
+					ISMC->SetupAttachment(this);
 					ISMC->RegisterComponent();
 					ActorOwner->AddInstanceComponent(ISMC);
 
