ChaosCachingChaosCaching
-- implemented secondary notifications for collsion and trailing events recording and playback.

diff --git a/Engine/Plugins/Experimental/ChaosCaching/Source/ChaosCaching/Private/Chaos/Adapters/GeometryCollectionComponentCacheAdapter.cpp b/Engine/Plugins/Experimental/ChaosCaching/Source/ChaosCaching/Private/Chaos/Adapters/GeometryCollectionComponentCacheAdapter.cpp
--- a/Engine/Plugins/Experimental/ChaosCaching/Source/ChaosCaching/Private/Chaos/Adapters/GeometryCollectionComponentCacheAdapter.cpp
+++ b/Engine/Plugins/Experimental/ChaosCaching/Source/ChaosCaching/Private/Chaos/Adapters/GeometryCollectionComponentCacheAdapter.cpp
@@ -11,6 +11,8 @@
 
 FName FEnableStateEvent::EventName("GC_Enable");
 FName FBreakingEvent::EventName("GC_Breaking");
+FName FCollisionEvent::EventName("GC_Collision");
+FName FTrailingEvent::EventName("GC_Trailing");
 
 namespace Chaos
 {
@@ -47,12 +49,19 @@ namespace Chaos
 
 		UGeometryCollectionComponent*    Comp  = CastChecked<UGeometryCollectionComponent>(InComp);
 		FGeometryCollectionPhysicsProxy* Proxy = Comp->GetPhysicsProxy();
-
+		
 		if(!Proxy)
 		{
 			return;
 		}
 
+		if (!CachedData.Contains(Proxy))
+		{
+			return;
+		}
+
+		const FCachedEventData& ProxyCachedEventData = CachedData[Proxy];
+
 		const Chaos::FPhysicsSolver* Solver         = Proxy->GetSolver<Chaos::FPhysicsSolver>();
 		const FGeometryCollection*   RestCollection = Proxy->GetSimParameters().RestCollection;
 
@@ -121,14 +130,14 @@ namespace Chaos
 			}
 		}
 
-		if (BreakingDataArray && ProxyBreakingDataIndices)
+		if (BreakingDataArray && ProxyCachedEventData.ProxyBreakingDataIndices)
 		{
-			for (int32 Index : *ProxyBreakingDataIndices)
+			for (int32 Index : *ProxyCachedEventData.ProxyBreakingDataIndices)
 			{
 				if (BreakingDataArray->IsValidIndex(Index))
 				{
 					const FBreakingData& BreakingData = (*BreakingDataArray)[Index];
-					if (const TPBDRigidParticleHandle<FReal, 3>*Rigid = BreakingData.Particle->CastToRigidParticle())
+					if (const TPBDRigidParticleHandle<FReal, 3>* Rigid = BreakingData.Particle->CastToRigidParticle())
 					{
 						int32 TransformIndex = Proxy->GetTransformGroupIndexFromHandle(Rigid);
 						if (TransformIndex > INDEX_NONE)
@@ -140,6 +149,46 @@ namespace Chaos
 				
 			}
 		}
+		
+		if (CollisionDataArray && ProxyCachedEventData.ProxyCollisionDataIndices)
+		{
+			for (int32 Index : *ProxyCachedEventData.ProxyCollisionDataIndices)
+			{
+				if (CollisionDataArray->IsValidIndex(Index))
+				{
+					const TCollisionData<float, 3>& CollisionData = (*CollisionDataArray)[Index];
+					if (const TPBDRigidParticleHandle<FReal, 3>* Rigid = CollisionData.Levelset->CastToRigidParticle())
+					{
+						int32 TransformIndex = Proxy->GetTransformGroupIndexFromHandle(Rigid);
+						if (TransformIndex > INDEX_NONE)
+						{
+							OutFrame.PushEvent(FCollisionEvent::EventName, InTime, FCollisionEvent(TransformIndex, CollisionData));
+						}
+					}
+				}
+
+			}
+		}
+		
+		if (TrailingDataArray && ProxyCachedEventData.ProxyTrailingDataIndices)
+		{
+			for (int32 Index : *ProxyCachedEventData.ProxyTrailingDataIndices)
+			{
+				if (TrailingDataArray->IsValidIndex(Index))
+				{
+					const TTrailingData<float, 3>& TrailingData = (*TrailingDataArray)[Index];
+					if (const TPBDRigidParticleHandle<FReal, 3>*Rigid = TrailingData.Particle->CastToRigidParticle())
+					{
+						int32 TransformIndex = Proxy->GetTransformGroupIndexFromHandle(Rigid);
+						if (TransformIndex > INDEX_NONE)
+						{
+							OutFrame.PushEvent(FTrailingEvent::EventName, InTime, FTrailingEvent(TransformIndex, TrailingData));
+						}
+					}
+				}
+
+			}
+		}
 
 		// Never going to change again till freed after writing to the cache so free up the extra space we reserved
 		OutFrame.PendingParticleData.Shrink();
@@ -185,6 +234,8 @@ namespace Chaos
 		const int32                      NumEventTracks = EvaluatedResult.Events.Num();
 		const TArray<FCacheEventHandle>* EnableEvents   = EvaluatedResult.Events.Find(FEnableStateEvent::EventName);
 		const TArray<FCacheEventHandle>* BreakingEvents = EvaluatedResult.Events.Find(FBreakingEvent::EventName);
+		const TArray<FCacheEventHandle>* CollisionEvents = EvaluatedResult.Events.Find(FCollisionEvent::EventName);
+		const TArray<FCacheEventHandle>* TrailingEvents = EvaluatedResult.Events.Find(FTrailingEvent::EventName);
 
 		if(EnableEvents)
 		{
@@ -302,6 +353,111 @@ namespace Chaos
 			}
 		}
 
+		if (TrailingEvents)
+		{
+			const FSolverTrailingEventFilter* SolverTrailingEventFilter = Solver->GetEventFilters()->GetTrailingFilter();
+
+			for (const FCacheEventHandle& Handle : *TrailingEvents)
+			{
+				if (FTrailingEvent* Event = Handle.Get<FTrailingEvent>())
+				{
+					if (Particles.IsValidIndex(Event->Index))
+					{
+						Chaos::TPBDRigidClusteredParticleHandle<float, 3>* Particle = Particles[Event->Index];
+
+						if (Particle)
+						{
+							if (Particle->ObjectState() != EObjectStateType::Kinematic)
+							{
+								// If a field or other external actor set the particle to static or dynamic we no longer apply the cache
+								continue;
+							}
+
+							TTrailingData<float, 3> CachedTrail;
+							CachedTrail.Particle = Particle;
+							CachedTrail.Location = Event->Location;
+							CachedTrail.Velocity = Event->Velocity;
+							CachedTrail.AngularVelocity = Event->AngularVelocity;
+							CachedTrail.BoundingBox = TAABB<float, 3>(Event->BoundingBoxMin, Event->BoundingBoxMax);
+
+							if (!SolverTrailingEventFilter->Enabled() || SolverTrailingEventFilter->Pass(CachedTrail))
+							{
+								float TimeStamp = Solver->GetSolverTime();
+								Solver->GetEventManager()->AddEvent<FTrailingEventData>(EEventType::Trailing, [&CachedTrail , TimeStamp](FTrailingEventData& TrailingEventData)
+									{
+										if (TrailingEventData.TrailingData.TimeCreated != TimeStamp)
+										{
+											TrailingEventData.TrailingData.AllTrailingsArray.Reset();
+											TrailingEventData.TrailingData.TimeCreated = TimeStamp;
+										}
+										TrailingEventData.TrailingData.AllTrailingsArray.Add(CachedTrail);
+									});
+							}
+						}
+					}
+				}
+			}
+		}
+
+		if (CollisionEvents)
+		{
+
+			const FSolverCollisionEventFilter* SolverCollisionEventFilter = Solver->GetEventFilters()->GetCollisionFilter();
+			for (const FCacheEventHandle& Handle : *CollisionEvents)
+			{
+				if (FCollisionEvent* Event = Handle.Get<FCollisionEvent>())
+				{
+					if (Particles.IsValidIndex(Event->Index))
+					{
+						Chaos::TPBDRigidClusteredParticleHandle<float, 3>* Particle = Particles[Event->Index];
+
+						if (Particle)
+						{
+							if (Particle->ObjectState() != EObjectStateType::Kinematic)
+							{
+								// If a field or other external actor set the particle to static or dynamic we no longer apply the cache
+								continue;
+							}
+
+							TCollisionData<float, 3> CachedCollision;
+							CachedCollision.Location = Event->Location;
+							CachedCollision.AccumulatedImpulse = Event->AccumulatedImpulse;
+							CachedCollision.Normal = Event->Normal;
+							CachedCollision.Velocity1 = Event->Velocity1;
+							CachedCollision.Velocity2 = Event->Velocity2;
+							CachedCollision.DeltaVelocity1 = Event->DeltaVelocity1;
+							CachedCollision.DeltaVelocity2 = Event->DeltaVelocity2;
+							CachedCollision.AngularVelocity1 = Event->AngularVelocity1;
+							CachedCollision.AngularVelocity2 = Event->AngularVelocity2;
+							CachedCollision.Mass1 = Event->Mass1;
+							CachedCollision.Mass2 = Event->Mass2;
+							CachedCollision.PenetrationDepth = Event->PenetrationDepth;
+							CachedCollision.Particle = Particle;
+
+							// #todo: Are these even available from a cache?
+							CachedCollision.Levelset = nullptr;
+
+							if (!SolverCollisionEventFilter->Enabled() || SolverCollisionEventFilter->Pass(CachedCollision))
+							{
+								float TimeStamp = Solver->GetSolverTime();
+								Solver->GetEventManager()->AddEvent<FCollisionEventData>(EEventType::Collision, [&CachedCollision, TimeStamp, Particle](FCollisionEventData& CollisionEventData)
+									{
+										if (CollisionEventData.CollisionData.TimeCreated != TimeStamp)
+										{
+											CollisionEventData.CollisionData.AllCollisionsArray.Reset();
+											CollisionEventData.PhysicsProxyToCollisionIndices.Reset();
+											CollisionEventData.CollisionData.TimeCreated = TimeStamp;
+										}
+										int32 NewIdx = CollisionEventData.CollisionData.AllCollisionsArray.Add(CachedCollision);
+										CollisionEventData.PhysicsProxyToCollisionIndices.PhysicsProxyToIndicesMap.FindOrAdd(Particle->PhysicsProxy()).Add(NewIdx);
+									});
+							}
+						}
+					}
+				}
+			}
+		}
+
 		const int32 NumTransforms = EvaluatedResult.Transform.Num();
 		for(int32 Index = 0; Index < NumTransforms; ++Index)
 		{
@@ -421,34 +577,42 @@ namespace Chaos
 	{
 		UGeometryCollectionComponent*    Comp     = CastChecked<UGeometryCollectionComponent>(InComponent);
 		FGeometryCollectionPhysicsProxy* Proxy    = Comp->GetPhysicsProxy();
-		
-		ProxyKey = Proxy;
-		BreakingDataArray = nullptr;
-		ProxyBreakingDataIndices = nullptr;
 
-		if(!Proxy)
+		if (!Proxy)
 		{
 			return false;
 		}
 
 		Chaos::FPhysicsSolver* Solver = Proxy->GetSolver<Chaos::FPhysicsSolver>();
 
-		if(!Solver)
+		if (!Solver)
 		{
 			return false;
 		}
 
-		// We need breaking data to record cluster breaking information into the cache
+		// We need secondary event data to record event information into the cache
 		Solver->SetGenerateBreakingData(true);
-
-
-		// Initialize event handlers to capture Breaking, Collision and Trailing events.
-		Chaos::FEventManager* EventManager = Solver->GetEventManager();
-		if (EventManager)
+		Solver->SetGenerateCollisionData(true);
+		Solver->SetGenerateTrailingData(true);
+		
+		// We only need to register event handlers once, the first time we initialize.
+		if (CachedData.Num() == 0)
 		{
-			EventManager->RegisterHandler<Chaos::FBreakingEventData>(Chaos::EEventType::Breaking, const_cast<FGeometryCollectionCacheAdapter*>(this), &FGeometryCollectionCacheAdapter::HandleBreakingEvents);
+			Chaos::FEventManager* EventManager = Solver->GetEventManager();
+			if (EventManager)
+			{
+				EventManager->RegisterHandler<Chaos::FBreakingEventData>(Chaos::EEventType::Breaking, const_cast<FGeometryCollectionCacheAdapter*>(this), &FGeometryCollectionCacheAdapter::HandleBreakingEvents);
+				EventManager->RegisterHandler<Chaos::FCollisionEventData>(Chaos::EEventType::Collision, const_cast<FGeometryCollectionCacheAdapter*>(this), &FGeometryCollectionCacheAdapter::HandleCollisionEvents);
+				EventManager->RegisterHandler<Chaos::FTrailingEventData>(Chaos::EEventType::Trailing, const_cast<FGeometryCollectionCacheAdapter*>(this), &FGeometryCollectionCacheAdapter::HandleTrailingEvents);
+			}
+
+			BreakingDataArray = nullptr;
+			CollisionDataArray = nullptr;
+			TrailingDataArray = nullptr;
 		}
 
+		CachedData.Add(Proxy, FCachedEventData());
+
 		return true;
 	}
 
@@ -469,12 +633,51 @@ namespace Chaos
 
 	void FGeometryCollectionCacheAdapter::HandleBreakingEvents(const Chaos::FBreakingEventData& Event)
 	{
-		if (ProxyKey)
+		BreakingDataArray = &Event.BreakingData.AllBreakingsArray;
+
+		for (TPair<IPhysicsProxyBase*, FCachedEventData>& Data : CachedData)
 		{
-			BreakingDataArray = &Event.BreakingData.AllBreakingsArray;
-			if (Event.PhysicsProxyToBreakingIndices.PhysicsProxyToIndicesMap.Contains(ProxyKey))
+			if (Event.PhysicsProxyToBreakingIndices.PhysicsProxyToIndicesMap.Contains(Data.Key))
+			{
+				Data.Value.ProxyBreakingDataIndices = &Event.PhysicsProxyToBreakingIndices.PhysicsProxyToIndicesMap[Data.Key];
+			}
+			else
+			{
+				Data.Value.ProxyBreakingDataIndices = nullptr;
+			}
+		}
+	}
+
+	void FGeometryCollectionCacheAdapter::HandleCollisionEvents(const Chaos::FCollisionEventData& Event)
+	{
+		CollisionDataArray = &Event.CollisionData.AllCollisionsArray;
+
+		for (TPair<IPhysicsProxyBase*, FCachedEventData>& Data : CachedData)
+		{
+			if (Event.PhysicsProxyToCollisionIndices.PhysicsProxyToIndicesMap.Contains(Data.Key))
+			{
+				Data.Value.ProxyCollisionDataIndices = &Event.PhysicsProxyToCollisionIndices.PhysicsProxyToIndicesMap[Data.Key];
+			}
+			else
+			{
+				Data.Value.ProxyCollisionDataIndices = nullptr;
+			}
+		}
+	}
+
+	void FGeometryCollectionCacheAdapter::HandleTrailingEvents(const Chaos::FTrailingEventData& Event)
+	{
+		TrailingDataArray = &Event.TrailingData.AllTrailingsArray;
+
+		for (TPair<IPhysicsProxyBase*, FCachedEventData>& Data : CachedData)
+		{
+			if (Event.PhysicsProxyToTrailingIndices.PhysicsProxyToIndicesMap.Contains(Data.Key))
+			{
+				Data.Value.ProxyTrailingDataIndices = &Event.PhysicsProxyToTrailingIndices.PhysicsProxyToIndicesMap[Data.Key];
+			}
+			else
 			{
-				ProxyBreakingDataIndices = &Event.PhysicsProxyToBreakingIndices.PhysicsProxyToIndicesMap[ProxyKey];
+				Data.Value.ProxyTrailingDataIndices = nullptr;
 			}
 		}
 	}
diff --git a/Engine/Plugins/Experimental/ChaosCaching/Source/ChaosCaching/Public/Chaos/Adapters/GeometryCollectionComponentCacheAdapter.h b/Engine/Plugins/Experimental/ChaosCaching/Source/ChaosCaching/Public/Chaos/Adapters/GeometryCollectionComponentCacheAdapter.h
--- a/Engine/Plugins/Experimental/ChaosCaching/Source/ChaosCaching/Public/Chaos/Adapters/GeometryCollectionComponentCacheAdapter.h
+++ b/Engine/Plugins/Experimental/ChaosCaching/Source/ChaosCaching/Public/Chaos/Adapters/GeometryCollectionComponentCacheAdapter.h
@@ -48,8 +48,7 @@ struct FBreakingEvent : public FCacheEventBase
 		, Mass(1.0f)
 		, BoundingBoxMin(0.0f, 0.0f, 0.0f)
 		, BoundingBoxMax(0.0f, 0.0f, 0.0f)
-	{
-	}
+	{ }
 
 	FBreakingEvent(int32 InIndex, const Chaos::FBreakingData& InData)
 		: Index(InIndex)
@@ -59,8 +58,7 @@ struct FBreakingEvent : public FCacheEventBase
 		, Mass(InData.Mass)
 		, BoundingBoxMin(InData.BoundingBox.Min())
 		, BoundingBoxMax(InData.BoundingBox.Max())
-	{
-	}
+	{ }
 
 	UPROPERTY()
 	int32 Index;
@@ -84,20 +82,151 @@ struct FBreakingEvent : public FCacheEventBase
 	FVector BoundingBoxMax;
 };
 
+USTRUCT()
+struct FCollisionEvent : public FCacheEventBase
+{
+	GENERATED_BODY()
+
+	static FName EventName;
+
+	FCollisionEvent()
+		: Index(INDEX_NONE)
+		, Location(0.0f, 0.0f, 0.0f)
+		, AccumulatedImpulse(0.0f, 0.0f, 0.0f)
+		, Normal(0.0f, 0.0f, 1.0f)
+		, Velocity1(0.0f, 0.0f, 0.0f)
+		, Velocity2(0.0f, 0.0f, 0.0f)
+		, DeltaVelocity1(0.0f, 0.0f, 0.0f)
+		, DeltaVelocity2(0.0f, 0.0f, 0.0f)
+		, AngularVelocity1(0.0f, 0.0f, 0.0f)
+		, AngularVelocity2(0.0f, 0.0f, 0.0f)
+		, Mass1(1.0f)
+		, Mass2(1.0f)
+		, PenetrationDepth(0.0f)
+	{}
+
+	FCollisionEvent(int32 InIndex, const Chaos::TCollisionData<float, 3>& InData)
+		: Index(InIndex)
+		, Location(InData.Location)
+		, AccumulatedImpulse(InData.AccumulatedImpulse)
+		, Normal(InData.Normal)
+		, Velocity1(InData.Velocity1)
+		, Velocity2(InData.Velocity2)
+		, DeltaVelocity1(InData.DeltaVelocity1)
+		, DeltaVelocity2(InData.DeltaVelocity2)
+		, AngularVelocity1(InData.AngularVelocity1)
+		, AngularVelocity2(InData.AngularVelocity2)
+		, Mass1(InData.Mass1)
+		, Mass2(InData.Mass2)
+		, PenetrationDepth(InData.PenetrationDepth)
+	{}
+
+	UPROPERTY()
+	int32 Index; 
+		
+	UPROPERTY()
+	FVector Location;
+	
+	UPROPERTY()
+	FVector AccumulatedImpulse;
+	
+	UPROPERTY()
+	FVector Normal;
+	
+	UPROPERTY()
+	FVector Velocity1;
+	
+	UPROPERTY()
+	FVector Velocity2;
+	
+	UPROPERTY()
+	FVector DeltaVelocity1;
+	
+	UPROPERTY()
+	FVector DeltaVelocity2;
+	
+	UPROPERTY()
+	FVector AngularVelocity1;
+	
+	UPROPERTY()
+	FVector AngularVelocity2;
+	
+	UPROPERTY()
+	float Mass1;
+	
+	UPROPERTY()
+	float Mass2;
+	
+	UPROPERTY()
+	float PenetrationDepth;
+};
+
+USTRUCT()
+struct FTrailingEvent : public FCacheEventBase
+{
+	GENERATED_BODY()
+
+		static FName EventName;
+
+	FTrailingEvent()
+		: Index(INDEX_NONE)
+		, Location(0.0f, 0.0f, 0.0f)
+		, Velocity(0.0f, 0.0f, 0.0f)
+		, AngularVelocity(0.0f, 0.0f, 0.0f)
+		, BoundingBoxMin(0.0f, 0.0f, 0.0f)
+		, BoundingBoxMax(0.0f, 0.0f, 0.0f)
+	{ }
+
+	FTrailingEvent(int32 InIndex, const Chaos::TTrailingData<float, 3>& InData)
+		: Index(InIndex)
+		, Location(InData.Location)
+		, Velocity(InData.Velocity)
+		, AngularVelocity(InData.AngularVelocity)
+		, BoundingBoxMin(InData.BoundingBox.Min())
+		, BoundingBoxMax(InData.BoundingBox.Max())
+	{ }
+
+	UPROPERTY()
+		int32 Index;
+
+	UPROPERTY()
+		FVector Location;
+
+	UPROPERTY()
+		FVector Velocity;
+
+	UPROPERTY()
+		FVector AngularVelocity;
+
+	UPROPERTY()
+		FVector BoundingBoxMin;
+
+	UPROPERTY()
+		FVector BoundingBoxMax;
+};
+
 
 class UChaosCache;
 class UPrimitiveComponent;
 
 namespace Chaos
 {
+	struct FCachedEventData
+	{
+		FCachedEventData()
+			: ProxyBreakingDataIndices(nullptr)
+			, ProxyCollisionDataIndices(nullptr)
+			, ProxyTrailingDataIndices(nullptr)
+		{}
+
+		const TArray<int32>* ProxyBreakingDataIndices;
+		const TArray<int32>* ProxyCollisionDataIndices;
+		const TArray<int32>* ProxyTrailingDataIndices;
+	};
+	
 	class FGeometryCollectionCacheAdapter : public FComponentCacheAdapter
 	{
 	public:
-		FGeometryCollectionCacheAdapter()
-			: ProxyKey(nullptr)
-			, BreakingDataArray(nullptr)
-			, ProxyBreakingDataIndices(nullptr)
-		{}
 
 		virtual ~FGeometryCollectionCacheAdapter() = default;
 
@@ -121,11 +250,15 @@ namespace Chaos
 
 	protected:
 		void HandleBreakingEvents(const Chaos::FBreakingEventData& Event);
+		void HandleCollisionEvents(const Chaos::FCollisionEventData& Event);
+		void HandleTrailingEvents(const Chaos::FTrailingEventData& Event);
 	
 	private:
-		const IPhysicsProxyBase* ProxyKey;
-		const Chaos::FBreakingDataArray* BreakingDataArray;
-		const TArray<int32>* ProxyBreakingDataIndices;
+		TMap<IPhysicsProxyBase*, FCachedEventData> CachedData;
+		const FBreakingDataArray* BreakingDataArray;
+		const FCollisionDataArray* CollisionDataArray;
+		const FTrailingDataArray* TrailingDataArray;
+
 
 	};
 }    // namespace Chaos
diff --git a/Engine/Source/Runtime/Experimental/Chaos/Private/EventDefaults.cpp b/Engine/Source/Runtime/Experimental/Chaos/Private/EventDefaults.cpp
--- a/Engine/Source/Runtime/Experimental/Chaos/Private/EventDefaults.cpp
+++ b/Engine/Source/Runtime/Experimental/Chaos/Private/EventDefaults.cpp
@@ -40,12 +40,16 @@ namespace Chaos
 			FCollisionDataArray& AllCollisionsDataArray = CollisionEventData.CollisionData.AllCollisionsArray;
 			TMap<IPhysicsProxyBase*, TArray<int32>>& AllCollisionsIndicesByPhysicsProxy = CollisionEventData.PhysicsProxyToCollisionIndices.PhysicsProxyToIndicesMap;
 
-			AllCollisionsDataArray.Reset();
-			AllCollisionsIndicesByPhysicsProxy.Reset();
-
-			CollisionEventData.CollisionData.TimeCreated = Solver->MTime;
-			CollisionEventData.PhysicsProxyToCollisionIndices.TimeCreated = Solver->MTime;
+			
+			if (CollisionEventData.CollisionData.TimeCreated != Solver->MTime)
+			{
+				AllCollisionsDataArray.Reset();
+				AllCollisionsIndicesByPhysicsProxy.Reset();
 
+				CollisionEventData.CollisionData.TimeCreated = Solver->MTime;
+				CollisionEventData.PhysicsProxyToCollisionIndices.TimeCreated = Solver->MTime;
+			}
+			
 			const auto* Evolution = Solver->GetEvolution();
 
 			const FPBDCollisionConstraints& CollisionRule = Evolution->GetCollisionConstraints();
@@ -318,35 +322,45 @@ namespace Chaos
 #if TODO_REIMPLEMENT_RIGID_CLUSTERING
 			const TMap<uint32, TUniquePtr<TArray<uint32>>>& ParentToChildrenMap = Evolution->GetRigidClustering().GetChildrenMap();
 #endif
-			auto& AllTrailingsDataArray = TrailingEventData.TrailingData.AllTrailingsArray;
+			FTrailingDataArray& AllTrailingsDataArray = TrailingEventData.TrailingData.AllTrailingsArray;
+			TMap<IPhysicsProxyBase*, TArray<int32>>& AllTrailingIndicesByPhysicsProxy = TrailingEventData.PhysicsProxyToTrailingIndices.PhysicsProxyToIndicesMap;
 
-			AllTrailingsDataArray.Reset();
+			if (TrailingEventData.TrailingData.TimeCreated != Solver->MTime)
+			{
+				AllTrailingsDataArray.Reset();
+				AllTrailingIndicesByPhysicsProxy.Reset();
 
-			TrailingEventData.TrailingData.TimeCreated = Solver->MTime;
+				TrailingEventData.TrailingData.TimeCreated = Solver->MTime;
+				TrailingEventData.PhysicsProxyToTrailingIndices.TimeCreated = Solver->MTime;
+			}
 
-			for (auto& ActiveParticle : Evolution->GetParticles().GetActiveParticlesView())
+			const TArray<TPBDRigidParticleHandle<Chaos::FReal, 3>*>& ActiveParticlesArray = Evolution->GetParticles().GetActiveParticlesArray();
+
+			for (TPBDRigidParticleHandle<Chaos::FReal, 3>* ActiveParticle : ActiveParticlesArray)
 			{
 
-				if (ensure(FMath::IsFinite(ActiveParticle.InvM())))
+				if (ensure(FMath::IsFinite(ActiveParticle->InvM())))
 				{
-					if (ActiveParticle.InvM() != 0.f &&
-						ActiveParticle.Geometry() &&
-						ActiveParticle.Geometry()->HasBoundingBox())
+					if (ActiveParticle->InvM() != 0.f &&
+						ActiveParticle->Geometry() &&
+						ActiveParticle->Geometry()->HasBoundingBox())
 					{
-						if (ensure(!ActiveParticle.X().ContainsNaN() &&
-							!ActiveParticle.V().ContainsNaN() &&
-							!ActiveParticle.W().ContainsNaN() &&
-							FMath::IsFinite(ActiveParticle.M())))
+						if (ensure(!ActiveParticle->X().ContainsNaN() &&
+							!ActiveParticle->V().ContainsNaN() &&
+							!ActiveParticle->W().ContainsNaN() &&
+							FMath::IsFinite(ActiveParticle->M())))
 						{
 							FTrailingData TrailingData;
-							TrailingData.Location = ActiveParticle.X();
-							TrailingData.Velocity = ActiveParticle.V();
-							TrailingData.AngularVelocity = ActiveParticle.W();
-							TrailingData.Mass = ActiveParticle.M();
-							TrailingData.Particle = nullptr; // #todo: provide a particle
-							if (ActiveParticle.Geometry()->HasBoundingBox())
+							TrailingData.Location = ActiveParticle->X();
+							TrailingData.Velocity = ActiveParticle->V();
+							TrailingData.AngularVelocity = ActiveParticle->W();
+							TrailingData.Mass = ActiveParticle->M();
+
+							TrailingData.Particle = ActiveParticle;
+							
+							if (ActiveParticle->Geometry()->HasBoundingBox())
 							{
-								TrailingData.BoundingBox = ActiveParticle.Geometry()->BoundingBox();
+								TrailingData.BoundingBox = ActiveParticle->Geometry()->BoundingBox();
 							}
 
 							const FSolverTrailingEventFilter* SolverTrailingEventFilter = Solver->GetEventFilters()->GetTrailingFilter();
@@ -356,6 +370,9 @@ namespace Chaos
 								FTrailingData& TrailingDataArrayItem = AllTrailingsDataArray[NewIdx];
 								TrailingDataArrayItem = TrailingData;
 
+								// Add to AllTrailingIndicesByPhysicsProxy
+								AllTrailingIndicesByPhysicsProxy.FindOrAdd(TrailingData.Particle->PhysicsProxy()).Add(FEventManager::EncodeCollisionIndex(NewIdx, false));
+
 								// If IdxParticle is a cluster store an index for a mesh in this cluster
 #if 0
 								if (ClusterIdsArray[IdxParticle].NumChildren > 0)
diff --git a/Engine/Source/Runtime/Experimental/Chaos/Public/Chaos/PBDRigidsSOAs.h b/Engine/Source/Runtime/Experimental/Chaos/Public/Chaos/PBDRigidsSOAs.h
--- a/Engine/Source/Runtime/Experimental/Chaos/Public/Chaos/PBDRigidsSOAs.h
+++ b/Engine/Source/Runtime/Experimental/Chaos/Public/Chaos/PBDRigidsSOAs.h
@@ -578,6 +578,8 @@ public:
 	const TParticleView<FPBDRigidParticles>& GetActiveParticlesView() const { UpdateIfNeeded(); return ActiveParticlesView; }
 	TParticleView<FPBDRigidParticles>& GetActiveParticlesView() { UpdateIfNeeded(); return ActiveParticlesView; }
 
+	const TArray<FPBDRigidParticleHandle*>& GetActiveParticlesArray() const { return ActiveParticlesArray; }
+	
 	const TParticleView<FPBDRigidParticles>& GetDirtyParticlesView() const { UpdateIfNeeded(); return DirtyParticlesView; }
 	TParticleView<FPBDRigidParticles>& GetDirtyParticlesView() { UpdateIfNeeded(); return DirtyParticlesView; }
 
diff --git a/Engine/Source/Runtime/Experimental/Chaos/Public/EventsData.h b/Engine/Source/Runtime/Experimental/Chaos/Public/EventsData.h
--- a/Engine/Source/Runtime/Experimental/Chaos/Public/EventsData.h
+++ b/Engine/Source/Runtime/Experimental/Chaos/Public/EventsData.h
@@ -114,6 +114,7 @@ namespace Chaos
 		FTrailingEventData() {}
 
 		FAllTrailingData TrailingData;
+		FIndicesByPhysicsProxy PhysicsProxyToTrailingIndices;
 	};
 
 	struct FSleepingEventData
diff --git a/Engine/Source/Runtime/Experimental/ChaosSolverEngine/Private/Chaos/ChaosGameplayEventDispatcher.cpp b/Engine/Source/Runtime/Experimental/ChaosSolverEngine/Private/Chaos/ChaosGameplayEventDispatcher.cpp
--- a/Engine/Source/Runtime/Experimental/ChaosSolverEngine/Private/Chaos/ChaosGameplayEventDispatcher.cpp
+++ b/Engine/Source/Runtime/Experimental/ChaosSolverEngine/Private/Chaos/ChaosGameplayEventDispatcher.cpp
@@ -55,6 +55,26 @@ static void DispatchPendingBreakEvents(TArray<FChaosBreakEvent> const& Events, T
 	}
 }
 
+static void SetCollisionInfoFromComp(FRigidBodyCollisionInfo& Info, UPrimitiveComponent* Comp)
+{
+	if (Comp)
+	{
+		Info.Component = Comp;
+		Info.Actor = Comp->GetOwner();
+
+		const FBodyInstance* const BodyInst = Comp->GetBodyInstance();
+		Info.BodyIndex = BodyInst ? BodyInst->InstanceBodyIndex : INDEX_NONE;
+		Info.BoneName = BodyInst && BodyInst->BodySetup.IsValid() ? BodyInst->BodySetup->BoneName : NAME_None;
+	}
+	else
+	{
+		Info.Component = nullptr;
+		Info.Actor = nullptr;
+		Info.BodyIndex = INDEX_NONE;
+		Info.BoneName = NAME_None;
+	}
+}
+
 FCollisionNotifyInfo& UChaosGameplayEventDispatcher::GetPendingCollisionForContactPair(const void* P0, const void* P1, bool& bNewEntry)
 {
 	const FUniqueContactPairKey Key = { P0, P1 };
@@ -275,77 +295,86 @@ void UChaosGameplayEventDispatcher::HandleCollisionEvents(const Chaos::FCollisio
 								int32 CollisionIdx = Chaos::FEventManager::DecodeCollisionIndex(EncodedCollisionIdx, bSwapOrder);
 
 								Chaos::FCollidingData const& CollisionDataItem = CollisionData[CollisionIdx];
-								IPhysicsProxyBase* const PhysicsProxy1 = bSwapOrder ? CollisionDataItem.Particle->PhysicsProxy() : CollisionDataItem.Levelset->PhysicsProxy();
 
+								// If we have no LevelSet pointer (which can be the case if we're reading a cache), force which proxy we use.
+								bool bUseParticle = CollisionDataItem.Levelset ? bSwapOrder : true;
+								if (CollisionDataItem.Particle)
 								{
-									bool bNewEntry = false;
-									FCollisionNotifyInfo& NotifyInfo = GetPendingCollisionForContactPair(PhysicsProxy0, PhysicsProxy1, bNewEntry);
-
-									// #note: we only notify on the first contact, though we will still accumulate the impulse data from subsequent contacts
-									const FVector NormalImpulse = FVector::DotProduct(CollisionDataItem.AccumulatedImpulse, CollisionDataItem.Normal) * CollisionDataItem.Normal;	// project impulse along normal
-									const FVector FrictionImpulse = FVector(CollisionDataItem.AccumulatedImpulse) - NormalImpulse; // friction is component not along contact normal
-									NotifyInfo.RigidCollisionData.TotalNormalImpulse += NormalImpulse;
-									NotifyInfo.RigidCollisionData.TotalFrictionImpulse += FrictionImpulse;
+									IPhysicsProxyBase* const PhysicsProxy1 = bUseParticle ? CollisionDataItem.Particle->PhysicsProxy() : CollisionDataItem.Levelset->PhysicsProxy();
 
-									if (bNewEntry)
 									{
-										// fill in legacy contact data
-										NotifyInfo.bCallEvent0 = true;
-										// if Comp1 wants this event too, it will get its own pending collision entry, so we leave it false
-
-										NotifyInfo.Info0.SetFrom(Scene.GetBodyInstanceFromProxy(PhysicsProxy0));
-										NotifyInfo.Info1.SetFrom(Scene.GetBodyInstanceFromProxy(PhysicsProxy1));
-
-										FRigidBodyContactInfo& NewContact = NotifyInfo.RigidCollisionData.ContactInfos.AddZeroed_GetRef();
-										NewContact.ContactNormal = CollisionDataItem.Normal;
-										NewContact.ContactPosition = CollisionDataItem.Location;
-										NewContact.ContactPenetration = CollisionDataItem.PenetrationDepth;
-										// NewContact.PhysMaterial[1] UPhysicalMaterial required here
-									}
+										bool bNewEntry = false;
+										FCollisionNotifyInfo& NotifyInfo = GetPendingCollisionForContactPair(PhysicsProxy0, PhysicsProxy1, bNewEntry);
 
-								}
+										// #note: we only notify on the first contact, though we will still accumulate the impulse data from subsequent contacts
+										const FVector NormalImpulse = FVector::DotProduct(CollisionDataItem.AccumulatedImpulse, CollisionDataItem.Normal) * CollisionDataItem.Normal;	// project impulse along normal
+										const FVector FrictionImpulse = FVector(CollisionDataItem.AccumulatedImpulse) - NormalImpulse; // friction is component not along contact normal
+										NotifyInfo.RigidCollisionData.TotalNormalImpulse += NormalImpulse;
+										NotifyInfo.RigidCollisionData.TotalFrictionImpulse += FrictionImpulse;
 
-								if (HandlerSet.ChaosHandlers.Num() > 0)
-								{
-									bool bNewEntry = false;
-									FChaosPendingCollisionNotify& ChaosNotifyInfo = GetPendingChaosCollisionForContactPair(PhysicsProxy0, PhysicsProxy1, bNewEntry);
+										if (bNewEntry)
+										{
+											UPrimitiveComponent* const Comp1 = Scene.GetOwningComponent<UPrimitiveComponent>(PhysicsProxy1);
+
+											// fill in legacy contact data
+											NotifyInfo.bCallEvent0 = true;
+											// if Comp1 wants this event too, it will get its own pending collision entry, so we leave it false
+
+											SetCollisionInfoFromComp(NotifyInfo.Info0, Comp0);
+											SetCollisionInfoFromComp(NotifyInfo.Info1, Comp1);
+
+											FRigidBodyContactInfo& NewContact = NotifyInfo.RigidCollisionData.ContactInfos.AddZeroed_GetRef();
+											NewContact.ContactNormal = CollisionDataItem.Normal;
+											NewContact.ContactPosition = CollisionDataItem.Location;
+											NewContact.ContactPenetration = CollisionDataItem.PenetrationDepth;
+											// NewContact.PhysMaterial[1] UPhysicalMaterial required here
+										}
+
+									}
 
-									// #note: we only notify on the first contact, though we will still accumulate the impulse data from subsequent contacts
-									ChaosNotifyInfo.CollisionInfo.AccumulatedImpulse += CollisionDataItem.AccumulatedImpulse;
 
-									if (bNewEntry)
+									if (HandlerSet.ChaosHandlers.Num() > 0)
 									{
-										UPrimitiveComponent* const Comp1 = Scene.GetOwningComponent<UPrimitiveComponent>(PhysicsProxy1);
+										bool bNewEntry = false;
+										FChaosPendingCollisionNotify& ChaosNotifyInfo = GetPendingChaosCollisionForContactPair(PhysicsProxy0, PhysicsProxy1, bNewEntry);
 
-										// fill in Chaos contact data
-										ChaosNotifyInfo.CollisionInfo.Component = Comp0;
-										ChaosNotifyInfo.CollisionInfo.OtherComponent = Comp1;
-										ChaosNotifyInfo.CollisionInfo.Location = CollisionDataItem.Location;
-										ChaosNotifyInfo.NotifyRecipients = HandlerSet.ChaosHandlers;
+										// #note: we only notify on the first contact, though we will still accumulate the impulse data from subsequent contacts
+										ChaosNotifyInfo.CollisionInfo.AccumulatedImpulse += CollisionDataItem.AccumulatedImpulse;
 
-										if (bSwapOrder)
-										{
-											ChaosNotifyInfo.CollisionInfo.AccumulatedImpulse = -CollisionDataItem.AccumulatedImpulse;
-											ChaosNotifyInfo.CollisionInfo.Normal = -CollisionDataItem.Normal;
-
-											ChaosNotifyInfo.CollisionInfo.Velocity = CollisionDataItem.Velocity2;
-											ChaosNotifyInfo.CollisionInfo.OtherVelocity = CollisionDataItem.Velocity1;
-											ChaosNotifyInfo.CollisionInfo.AngularVelocity = CollisionDataItem.AngularVelocity2;
-											ChaosNotifyInfo.CollisionInfo.OtherAngularVelocity = CollisionDataItem.AngularVelocity1;
-											ChaosNotifyInfo.CollisionInfo.Mass = CollisionDataItem.Mass2;
-											ChaosNotifyInfo.CollisionInfo.OtherMass = CollisionDataItem.Mass1;
-										}
-										else
+										if (bNewEntry)
 										{
-											ChaosNotifyInfo.CollisionInfo.AccumulatedImpulse = CollisionDataItem.AccumulatedImpulse;
-											ChaosNotifyInfo.CollisionInfo.Normal = CollisionDataItem.Normal;
-
-											ChaosNotifyInfo.CollisionInfo.Velocity = CollisionDataItem.Velocity1;
-											ChaosNotifyInfo.CollisionInfo.OtherVelocity = CollisionDataItem.Velocity2;
-											ChaosNotifyInfo.CollisionInfo.AngularVelocity = CollisionDataItem.AngularVelocity1;
-											ChaosNotifyInfo.CollisionInfo.OtherAngularVelocity = CollisionDataItem.AngularVelocity2;
-											ChaosNotifyInfo.CollisionInfo.Mass = CollisionDataItem.Mass1;
-											ChaosNotifyInfo.CollisionInfo.OtherMass = CollisionDataItem.Mass2;
+											UPrimitiveComponent* const Comp1 = Scene.GetOwningComponent<UPrimitiveComponent>(PhysicsProxy1);
+
+											// fill in Chaos contact data
+											ChaosNotifyInfo.CollisionInfo.Component = Comp0;
+											ChaosNotifyInfo.CollisionInfo.OtherComponent = Comp1;
+											ChaosNotifyInfo.CollisionInfo.Location = CollisionDataItem.Location;
+											ChaosNotifyInfo.NotifyRecipients = HandlerSet.ChaosHandlers;
+
+											if (bSwapOrder)
+											{
+												ChaosNotifyInfo.CollisionInfo.AccumulatedImpulse = -CollisionDataItem.AccumulatedImpulse;
+												ChaosNotifyInfo.CollisionInfo.Normal = -CollisionDataItem.Normal;
+
+												ChaosNotifyInfo.CollisionInfo.Velocity = CollisionDataItem.Velocity2;
+												ChaosNotifyInfo.CollisionInfo.OtherVelocity = CollisionDataItem.Velocity1;
+												ChaosNotifyInfo.CollisionInfo.AngularVelocity = CollisionDataItem.AngularVelocity2;
+												ChaosNotifyInfo.CollisionInfo.OtherAngularVelocity = CollisionDataItem.AngularVelocity1;
+												ChaosNotifyInfo.CollisionInfo.Mass = CollisionDataItem.Mass2;
+												ChaosNotifyInfo.CollisionInfo.OtherMass = CollisionDataItem.Mass1;
+											}
+											else
+											{
+												ChaosNotifyInfo.CollisionInfo.AccumulatedImpulse = CollisionDataItem.AccumulatedImpulse;
+												ChaosNotifyInfo.CollisionInfo.Normal = CollisionDataItem.Normal;
+
+												ChaosNotifyInfo.CollisionInfo.Velocity = CollisionDataItem.Velocity1;
+												ChaosNotifyInfo.CollisionInfo.OtherVelocity = CollisionDataItem.Velocity2;
+												ChaosNotifyInfo.CollisionInfo.AngularVelocity = CollisionDataItem.AngularVelocity1;
+												ChaosNotifyInfo.CollisionInfo.OtherAngularVelocity = CollisionDataItem.AngularVelocity2;
+												ChaosNotifyInfo.CollisionInfo.Mass = CollisionDataItem.Mass1;
+												ChaosNotifyInfo.CollisionInfo.OtherMass = CollisionDataItem.Mass2;
+											}
 										}
 									}
 								}
