Chaos Caching
-- change TCollisionData to FCollingData and FTrailingData to FTrailingData

diff --git a/Engine/Plugins/Experimental/ChaosCaching/Source/ChaosCaching/Private/Chaos/Adapters/GeometryCollectionComponentCacheAdapter.cpp b/Engine/Plugins/Experimental/ChaosCaching/Source/ChaosCaching/Private/Chaos/Adapters/GeometryCollectionComponentCacheAdapter.cpp
--- a/Engine/Plugins/Experimental/ChaosCaching/Source/ChaosCaching/Private/Chaos/Adapters/GeometryCollectionComponentCacheAdapter.cpp
+++ b/Engine/Plugins/Experimental/ChaosCaching/Source/ChaosCaching/Private/Chaos/Adapters/GeometryCollectionComponentCacheAdapter.cpp
@@ -156,7 +156,7 @@ namespace Chaos
 			{
 				if (CollisionDataArray->IsValidIndex(Index))
 				{
-					const TCollisionData<float, 3>& CollisionData = (*CollisionDataArray)[Index];
+					const FCollidingData& CollisionData = (*CollisionDataArray)[Index];
 					if (const TPBDRigidParticleHandle<FReal, 3>* Rigid = CollisionData.Levelset->CastToRigidParticle())
 					{
 						int32 TransformIndex = Proxy->GetTransformGroupIndexFromHandle(Rigid);
@@ -176,7 +176,7 @@ namespace Chaos
 			{
 				if (TrailingDataArray->IsValidIndex(Index))
 				{
-					const TTrailingData<float, 3>& TrailingData = (*TrailingDataArray)[Index];
+					const FTrailingData& TrailingData = (*TrailingDataArray)[Index];
 					if (const TPBDRigidParticleHandle<FReal, 3>*Rigid = TrailingData.Particle->CastToRigidParticle())
 					{
 						int32 TransformIndex = Proxy->GetTransformGroupIndexFromHandle(Rigid);
@@ -373,7 +373,7 @@ namespace Chaos
 								continue;
 							}
 
-							TTrailingData<float, 3> CachedTrail;
+							FTrailingData CachedTrail;
 							CachedTrail.Particle = Particle;
 							CachedTrail.Location = Event->Location;
 							CachedTrail.Velocity = Event->Velocity;
@@ -419,7 +419,7 @@ namespace Chaos
 								continue;
 							}
 
-							TCollisionData<float, 3> CachedCollision;
+							FCollidingData CachedCollision;
 							CachedCollision.Location = Event->Location;
 							CachedCollision.AccumulatedImpulse = Event->AccumulatedImpulse;
 							CachedCollision.Normal = Event->Normal;
diff --git a/Engine/Plugins/Experimental/ChaosCaching/Source/ChaosCaching/Public/Chaos/Adapters/GeometryCollectionComponentCacheAdapter.h b/Engine/Plugins/Experimental/ChaosCaching/Source/ChaosCaching/Public/Chaos/Adapters/GeometryCollectionComponentCacheAdapter.h
--- a/Engine/Plugins/Experimental/ChaosCaching/Source/ChaosCaching/Public/Chaos/Adapters/GeometryCollectionComponentCacheAdapter.h
+++ b/Engine/Plugins/Experimental/ChaosCaching/Source/ChaosCaching/Public/Chaos/Adapters/GeometryCollectionComponentCacheAdapter.h
@@ -105,7 +105,7 @@ struct FCollisionEvent : public FCacheEventBase
 		, PenetrationDepth(0.0f)
 	{}
 
-	FCollisionEvent(int32 InIndex, const Chaos::TCollisionData<float, 3>& InData)
+	FCollisionEvent(int32 InIndex, const Chaos::FCollidingData& InData)
 		: Index(InIndex)
 		, Location(InData.Location)
 		, AccumulatedImpulse(InData.AccumulatedImpulse)
@@ -177,7 +177,7 @@ struct FTrailingEvent : public FCacheEventBase
 		, BoundingBoxMax(0.0f, 0.0f, 0.0f)
 	{ }
 
-	FTrailingEvent(int32 InIndex, const Chaos::TTrailingData<float, 3>& InData)
+	FTrailingEvent(int32 InIndex, const Chaos::FTrailingData& InData)
 		: Index(InIndex)
 		, Location(InData.Location)
 		, Velocity(InData.Velocity)
