ChaosCaching
-- implemented more robust warnings to alert user if invalid cache or components are being used.

diff --git a/Engine/Plugins/Experimental/ChaosCaching/Source/ChaosCaching/Private/Chaos/CacheManagerActor.cpp b/Engine/Plugins/Experimental/ChaosCaching/Source/ChaosCaching/Private/Chaos/CacheManagerActor.cpp
--- a/Engine/Plugins/Experimental/ChaosCaching/Source/ChaosCaching/Private/Chaos/CacheManagerActor.cpp
+++ b/Engine/Plugins/Experimental/ChaosCaching/Source/ChaosCaching/Private/Chaos/CacheManagerActor.cpp
@@ -186,6 +186,8 @@ void AChaosCacheManager::BeginPlay()
 
 	if(!CacheCollection)
 	{
+		UE_LOG(LogChaosCache, Warning, TEXT("%s has no valid cache asset. Components will revert to dynamic simulation."), *GetName());
+		
 		// without a collection the cache manager can't do anything, no reason to inialise the observed array
 		SetActorTickEnabled(false);
 		return;
@@ -217,6 +219,7 @@ void AChaosCacheManager::BeginPlay()
 
 		if(!Comp)
 		{
+			UE_LOG(LogChaosCache, Warning, TEXT("%s has invalid observed component reference."), *GetName());
 			ActiveAdapters.Add(nullptr);
 			continue;
 		}
@@ -235,7 +238,9 @@ void AChaosCacheManager::BeginPlay()
 		Algo::Sort(DerivedAdapters, ByPriority);
 
 		if(DirectAdapters.Num() == 0 && DerivedAdapters.Num() == 0)
-		{
+		{	
+			UE_LOG(LogChaosCache, Warning, TEXT("%s observing component with no appropriate adapter."), *GetName());
+
 			// No actual adapter for this class type, log and push nullptr for this observed object
 			ActiveAdapters.Add(nullptr);
 		}
