ChaosCaching
-- Implemented adapter method to initialize the singleton to clear any global data that needs to be built.
-- Secondary event notifies, eg breaking, collision, trailing, are cached in component space to allow for cache repositioning on playbaack.

diff --git a/Engine/Plugins/Experimental/ChaosCaching/Source/ChaosCaching/Private/Chaos/Adapters/GeometryCollectionComponentCacheAdapter.cpp b/Engine/Plugins/Experimental/ChaosCaching/Source/ChaosCaching/Private/Chaos/Adapters/GeometryCollectionComponentCacheAdapter.cpp
--- a/Engine/Plugins/Experimental/ChaosCaching/Source/ChaosCaching/Private/Chaos/Adapters/GeometryCollectionComponentCacheAdapter.cpp
+++ b/Engine/Plugins/Experimental/ChaosCaching/Source/ChaosCaching/Private/Chaos/Adapters/GeometryCollectionComponentCacheAdapter.cpp
@@ -49,6 +49,7 @@ namespace Chaos
 
 		UGeometryCollectionComponent*    Comp  = CastChecked<UGeometryCollectionComponent>(InComp);
 		FGeometryCollectionPhysicsProxy* Proxy = Comp->GetPhysicsProxy();
+		const FTransform WorldToComponent = Comp->GetComponentTransform().Inverse();
 		
 		if(!Proxy)
 		{
@@ -142,7 +143,7 @@ namespace Chaos
 						int32 TransformIndex = Proxy->GetTransformGroupIndexFromHandle(Rigid);
 						if (TransformIndex > INDEX_NONE)
 						{
-							OutFrame.PushEvent(FBreakingEvent::EventName, InTime, FBreakingEvent(TransformIndex, BreakingData));
+							OutFrame.PushEvent(FBreakingEvent::EventName, InTime, FBreakingEvent(TransformIndex, BreakingData, WorldToComponent));
 						}
 					}
 				}
@@ -162,7 +163,7 @@ namespace Chaos
 						int32 TransformIndex = Proxy->GetTransformGroupIndexFromHandle(Rigid);
 						if (TransformIndex > INDEX_NONE)
 						{
-							OutFrame.PushEvent(FCollisionEvent::EventName, InTime, FCollisionEvent(TransformIndex, CollisionData));
+							OutFrame.PushEvent(FCollisionEvent::EventName, InTime, FCollisionEvent(TransformIndex, CollisionData, WorldToComponent));
 						}
 					}
 				}
@@ -182,7 +183,7 @@ namespace Chaos
 						int32 TransformIndex = Proxy->GetTransformGroupIndexFromHandle(Rigid);
 						if (TransformIndex > INDEX_NONE)
 						{
-							OutFrame.PushEvent(FTrailingEvent::EventName, InTime, FTrailingEvent(TransformIndex, TrailingData));
+							OutFrame.PushEvent(FTrailingEvent::EventName, InTime, FTrailingEvent(TransformIndex, TrailingData, WorldToComponent));
 						}
 					}
 				}
@@ -205,6 +206,7 @@ namespace Chaos
 
 		UGeometryCollectionComponent*    Comp  = CastChecked<UGeometryCollectionComponent>(InComponent);
 		FGeometryCollectionPhysicsProxy* Proxy = Comp->GetPhysicsProxy();
+		const FTransform ComponentToWorld = Comp->GetComponentTransform();
 
 		if(!Proxy)
 		{
@@ -328,11 +330,12 @@ namespace Chaos
 
 							FBreakingData CachedBreak;
 							CachedBreak.Particle = Particle;
-							CachedBreak.Location = Event->Location;
-							CachedBreak.Velocity = Event->Velocity;
+							CachedBreak.Location = ComponentToWorld.TransformPosition(Event->Location);
+							CachedBreak.Velocity = ComponentToWorld.TransformVector(Event->Velocity);
 							CachedBreak.AngularVelocity = Event->AngularVelocity;
 							CachedBreak.Mass = Event->Mass;
 							CachedBreak.BoundingBox = TAABB<float, 3>(Event->BoundingBoxMin, Event->BoundingBoxMax);
+							CachedBreak.BoundingBox = CachedBreak.BoundingBox.TransformedAABB(ComponentToWorld);
 
 							if (!SolverBreakingEventFilter->Enabled() || SolverBreakingEventFilter->Pass(CachedBreak))
 							{
@@ -375,10 +378,11 @@ namespace Chaos
 
 							FTrailingData CachedTrail;
 							CachedTrail.Particle = Particle;
-							CachedTrail.Location = Event->Location;
-							CachedTrail.Velocity = Event->Velocity;
+							CachedTrail.Location = ComponentToWorld.TransformPosition(Event->Location);
+							CachedTrail.Velocity = ComponentToWorld.TransformVector(Event->Velocity);
 							CachedTrail.AngularVelocity = Event->AngularVelocity;
 							CachedTrail.BoundingBox = TAABB<float, 3>(Event->BoundingBoxMin, Event->BoundingBoxMax);
+							CachedTrail.BoundingBox = CachedTrail.BoundingBox.TransformedAABB(ComponentToWorld);
 
 							if (!SolverTrailingEventFilter->Enabled() || SolverTrailingEventFilter->Pass(CachedTrail))
 							{
@@ -420,13 +424,13 @@ namespace Chaos
 							}
 
 							FCollidingData CachedCollision;
-							CachedCollision.Location = Event->Location;
-							CachedCollision.AccumulatedImpulse = Event->AccumulatedImpulse;
-							CachedCollision.Normal = Event->Normal;
-							CachedCollision.Velocity1 = Event->Velocity1;
-							CachedCollision.Velocity2 = Event->Velocity2;
-							CachedCollision.DeltaVelocity1 = Event->DeltaVelocity1;
-							CachedCollision.DeltaVelocity2 = Event->DeltaVelocity2;
+							CachedCollision.Location = ComponentToWorld.TransformPosition(Event->Location);
+							CachedCollision.AccumulatedImpulse = ComponentToWorld.TransformVector(Event->AccumulatedImpulse);
+							CachedCollision.Normal = ComponentToWorld.TransformVector(Event->Normal);
+							CachedCollision.Velocity1 = ComponentToWorld.TransformVector(Event->Velocity1);
+							CachedCollision.Velocity2 = ComponentToWorld.TransformVector(Event->Velocity2);
+							CachedCollision.DeltaVelocity1 = ComponentToWorld.TransformVector(Event->DeltaVelocity1);
+							CachedCollision.DeltaVelocity2 = ComponentToWorld.TransformVector(Event->DeltaVelocity2);
 							CachedCollision.AngularVelocity1 = Event->AngularVelocity1;
 							CachedCollision.AngularVelocity2 = Event->AngularVelocity2;
 							CachedCollision.Mass1 = Event->Mass1;
@@ -573,6 +577,11 @@ namespace Chaos
 		return nullptr;
 	}
 
+	void FGeometryCollectionCacheAdapter::Initialize()
+	{
+		CachedData.Empty();
+	}
+	
 	bool FGeometryCollectionCacheAdapter::InitializeForRecord(UPrimitiveComponent* InComponent, UChaosCache* InCache)
 	{
 		UGeometryCollectionComponent*    Comp     = CastChecked<UGeometryCollectionComponent>(InComponent);
diff --git a/Engine/Plugins/Experimental/ChaosCaching/Source/ChaosCaching/Private/Chaos/CacheManagerActor.cpp b/Engine/Plugins/Experimental/ChaosCaching/Source/ChaosCaching/Private/Chaos/CacheManagerActor.cpp
--- a/Engine/Plugins/Experimental/ChaosCaching/Source/ChaosCaching/Private/Chaos/CacheManagerActor.cpp
+++ b/Engine/Plugins/Experimental/ChaosCaching/Source/ChaosCaching/Private/Chaos/CacheManagerActor.cpp
@@ -203,6 +203,11 @@ void AChaosCacheManager::BeginPlay()
 	IModularFeatures&               ModularFeatures = IModularFeatures::Get();
 	TArray<FComponentCacheAdapter*> Adapters        = ModularFeatures.GetModularFeatureImplementations<FComponentCacheAdapter>(FComponentCacheAdapter::FeatureName);
 
+	for (FComponentCacheAdapter* Adapter : Adapters)
+	{
+		Adapter->Initialize();
+	}
+	
 	ActiveAdapters.Reset();
 
 	int32       NumFailedPlaybackEntries = 0;
diff --git a/Engine/Plugins/Experimental/ChaosCaching/Source/ChaosCaching/Public/Chaos/Adapters/CacheAdapter.h b/Engine/Plugins/Experimental/ChaosCaching/Source/ChaosCaching/Public/Chaos/Adapters/CacheAdapter.h
--- a/Engine/Plugins/Experimental/ChaosCaching/Source/ChaosCaching/Public/Chaos/Adapters/CacheAdapter.h
+++ b/Engine/Plugins/Experimental/ChaosCaching/Source/ChaosCaching/Public/Chaos/Adapters/CacheAdapter.h
@@ -113,6 +113,11 @@ namespace Chaos
 		 */
 		virtual Chaos::FPhysicsSolver* GetComponentSolver(UPrimitiveComponent* InComponent) const = 0;
 
+		/**
+		 * Called from the game thread to perform any global setup that the adapter may need to perform.
+		 */
+		virtual void Initialize() {}
+		
 		/**
 		 * Called from the game thread to initialize a component and cache ready to record a cache
 		 * @param InComponent Target component to initialize
diff --git a/Engine/Plugins/Experimental/ChaosCaching/Source/ChaosCaching/Public/Chaos/Adapters/GeometryCollectionComponentCacheAdapter.h b/Engine/Plugins/Experimental/ChaosCaching/Source/ChaosCaching/Public/Chaos/Adapters/GeometryCollectionComponentCacheAdapter.h
--- a/Engine/Plugins/Experimental/ChaosCaching/Source/ChaosCaching/Public/Chaos/Adapters/GeometryCollectionComponentCacheAdapter.h
+++ b/Engine/Plugins/Experimental/ChaosCaching/Source/ChaosCaching/Public/Chaos/Adapters/GeometryCollectionComponentCacheAdapter.h
@@ -50,15 +50,19 @@ struct FBreakingEvent : public FCacheEventBase
 		, BoundingBoxMax(0.0f, 0.0f, 0.0f)
 	{ }
 
-	FBreakingEvent(int32 InIndex, const Chaos::FBreakingData& InData)
+	FBreakingEvent(int32 InIndex, const Chaos::FBreakingData& InData, const FTransform& WorldToComponent)
 		: Index(InIndex)
-		, Location(InData.Location)
-		, Velocity(InData.Velocity)
+		, Location(WorldToComponent.TransformPosition(InData.Location))
+		, Velocity(WorldToComponent.TransformVector(InData.Velocity))
 		, AngularVelocity(InData.AngularVelocity)
 		, Mass(InData.Mass)
 		, BoundingBoxMin(InData.BoundingBox.Min())
 		, BoundingBoxMax(InData.BoundingBox.Max())
-	{ }
+	{
+		FBox TransformedBounds = FBox(BoundingBoxMin, BoundingBoxMax).TransformBy(WorldToComponent);
+		BoundingBoxMin = TransformedBounds.Min;
+		BoundingBoxMax = TransformedBounds.Max;
+	}
 
 	UPROPERTY()
 	int32 Index;
@@ -105,15 +109,15 @@ struct FCollisionEvent : public FCacheEventBase
 		, PenetrationDepth(0.0f)
 	{}
 
-	FCollisionEvent(int32 InIndex, const Chaos::FCollidingData& InData)
+	FCollisionEvent(int32 InIndex, const Chaos::FCollidingData& InData, const FTransform& WorldToComponent)
 		: Index(InIndex)
-		, Location(InData.Location)
-		, AccumulatedImpulse(InData.AccumulatedImpulse)
-		, Normal(InData.Normal)
-		, Velocity1(InData.Velocity1)
-		, Velocity2(InData.Velocity2)
-		, DeltaVelocity1(InData.DeltaVelocity1)
-		, DeltaVelocity2(InData.DeltaVelocity2)
+		, Location(WorldToComponent.TransformPosition(InData.Location))
+		, AccumulatedImpulse(WorldToComponent.TransformVector(InData.AccumulatedImpulse))
+		, Normal(WorldToComponent.TransformVector(InData.Normal))
+		, Velocity1(WorldToComponent.TransformVector(InData.Velocity1))
+		, Velocity2(WorldToComponent.TransformVector(InData.Velocity2))
+		, DeltaVelocity1(WorldToComponent.TransformVector(InData.DeltaVelocity1))
+		, DeltaVelocity2(WorldToComponent.TransformVector(InData.DeltaVelocity2))
 		, AngularVelocity1(InData.AngularVelocity1)
 		, AngularVelocity2(InData.AngularVelocity2)
 		, Mass1(InData.Mass1)
@@ -177,14 +181,18 @@ struct FTrailingEvent : public FCacheEventBase
 		, BoundingBoxMax(0.0f, 0.0f, 0.0f)
 	{ }
 
-	FTrailingEvent(int32 InIndex, const Chaos::FTrailingData& InData)
+	FTrailingEvent(int32 InIndex, const Chaos::FTrailingData& InData, const FTransform& WorldToComponent)
 		: Index(InIndex)
-		, Location(InData.Location)
-		, Velocity(InData.Velocity)
+		, Location(WorldToComponent.TransformPosition(InData.Location))
+		, Velocity(WorldToComponent.TransformVector(InData.Velocity))
 		, AngularVelocity(InData.AngularVelocity)
 		, BoundingBoxMin(InData.BoundingBox.Min())
 		, BoundingBoxMax(InData.BoundingBox.Max())
-	{ }
+	{
+		FBox TransformedBounds = FBox(BoundingBoxMin, BoundingBoxMax).TransformBy(WorldToComponent);
+		BoundingBoxMin = TransformedBounds.Min;
+		BoundingBoxMax = TransformedBounds.Max;
+	}
 
 	UPROPERTY()
 		int32 Index;
@@ -237,6 +245,7 @@ namespace Chaos
 		FGuid                  GetGuid() const override;
 		Chaos::FPhysicsSolver* GetComponentSolver(UPrimitiveComponent* InComponent) const override;
 		bool                   ValidForPlayback(UPrimitiveComponent* InComponent, UChaosCache* InCache) const override;
+		void				   Initialize() override;
 		bool                   InitializeForRecord(UPrimitiveComponent* InComponent, UChaosCache* InCache) override;
 		bool                   InitializeForPlayback(UPrimitiveComponent* InComponent, UChaosCache* InCache) const override;
 		void                   Record_PostSolve(UPrimitiveComponent* InComp, const FTransform& InRootTransform, FPendingFrameWrite& OutFrame, Chaos::FReal InTime) const override;
