Chaos : Caching
- Dirty caching on BeginPlay.

diff --git a/Engine/Plugins/Experimental/ChaosCaching/Source/ChaosCaching/Private/Chaos/ChaosCache.cpp b/Engine/Plugins/Experimental/ChaosCaching/Source/ChaosCaching/Private/Chaos/ChaosCache.cpp
--- a/Engine/Plugins/Experimental/ChaosCaching/Source/ChaosCaching/Private/Chaos/ChaosCache.cpp
+++ b/Engine/Plugins/Experimental/ChaosCaching/Source/ChaosCaching/Private/Chaos/ChaosCache.cpp
@@ -119,6 +119,11 @@ FCacheUserToken UChaosCache::BeginRecord(UPrimitiveComponent* InComponent, FGuid
 			// Build a compatibility hash for the component state.
 			AdapterGuid = InAdapterId;
 
+			if (UPackage* Package = GetOutermost())
+			{
+				Package->SetDirtyFlag(true);
+			}
+
 			return FCacheUserToken(true, true, this);
 		}
 		else
