ChaosCaching
-- Guarded a Proxy access where it's possible to have a nullptr.

diff --git a/Engine/Plugins/Experimental/ChaosCaching/Source/ChaosCaching/Private/Chaos/Adapters/GeometryCollectionComponentCacheAdapter.cpp b/Engine/Plugins/Experimental/ChaosCaching/Source/ChaosCaching/Private/Chaos/Adapters/GeometryCollectionComponentCacheAdapter.cpp
--- a/Engine/Plugins/Experimental/ChaosCaching/Source/ChaosCaching/Private/Chaos/Adapters/GeometryCollectionComponentCacheAdapter.cpp
+++ b/Engine/Plugins/Experimental/ChaosCaching/Source/ChaosCaching/Private/Chaos/Adapters/GeometryCollectionComponentCacheAdapter.cpp
@@ -628,16 +628,25 @@ namespace Chaos
 	bool FGeometryCollectionCacheAdapter::InitializeForPlayback(UPrimitiveComponent* InComponent, UChaosCache* InCache) const
 	{
 		UGeometryCollectionComponent*    Comp  = CastChecked<UGeometryCollectionComponent>(InComponent);
-		FGeometryCollectionPhysicsProxy* Proxy = Comp->GetPhysicsProxy();
+		if (!Comp)
+		{
+			return false;
+		}
 
-		FGeometryDynamicCollection& Collection = Proxy->GetPhysicsCollection();
+		FGeometryCollectionPhysicsProxy* Proxy = Comp->GetPhysicsProxy();
 
-		for (int32& State : Collection.DynamicState)
+		if (ensure(Proxy))
 		{
-			State = (int32)EObjectStateTypeEnum::Chaos_Object_Kinematic;
-		}
+			FGeometryDynamicCollection& Collection = Proxy->GetPhysicsCollection();
 
-		return true;
+			const int32 NumDynamicState = Collection.DynamicState.Num();
+			for (int32 StateIdx = 0; StateIdx < NumDynamicState; ++StateIdx)
+			{
+				Collection.DynamicState[StateIdx] = (int32)EObjectStateTypeEnum::Chaos_Object_Kinematic;
+			}
+		}
+		
+		return true;		
 	}
 
 	void FGeometryCollectionCacheAdapter::HandleBreakingEvents(const Chaos::FBreakingEventData& Event)
