Chaos Event Manager
-- Removed particle handle pointers from event payloads, as these are unsafe on the GT where they are often utilized.
-- Test proxy pointers when used on the game thread in case the proxy is pending destruction and is no longer valid.
-- Moved transform of particle handle -> transform group index for geometry collection event to PT where it is safe to reference the particle handle.

diff --git a/Engine/Plugins/Experimental/ChaosCaching/Source/ChaosCaching/Private/Chaos/Adapters/GeometryCollectionComponentCacheAdapter.cpp b/Engine/Plugins/Experimental/ChaosCaching/Source/ChaosCaching/Private/Chaos/Adapters/GeometryCollectionComponentCacheAdapter.cpp
--- a/Engine/Plugins/Experimental/ChaosCaching/Source/ChaosCaching/Private/Chaos/Adapters/GeometryCollectionComponentCacheAdapter.cpp
+++ b/Engine/Plugins/Experimental/ChaosCaching/Source/ChaosCaching/Private/Chaos/Adapters/GeometryCollectionComponentCacheAdapter.cpp
@@ -44,12 +44,16 @@ namespace Chaos
 
 	void FGeometryCollectionCacheAdapter::Record_PostSolve(UPrimitiveComponent* InComp, const FTransform& InRootTransform, FPendingFrameWrite& OutFrame, Chaos::FReal InTime) const
 	{
+		ensure(IsInGameThreadContext());
+
 		using FClusterParticle = Chaos::FPBDRigidClusteredParticleHandle;
 		using FRigidParticle = Chaos::FPBDRigidParticleHandle;
 
 		UGeometryCollectionComponent*    Comp  = CastChecked<UGeometryCollectionComponent>(InComp);
 		FGeometryCollectionPhysicsProxy* Proxy = Comp->GetPhysicsProxy();
 		const FTransform WorldToComponent = Comp->GetComponentTransform().Inverse();
+
+		FPhysScene_Chaos& Scene = *(Comp->GetWorld()->GetPhysicsScene());
 		
 		if(!Proxy)
 		{
@@ -82,22 +86,24 @@ namespace Chaos
 		// Pre-alloc once for worst case.
 		OutFrame.PendingParticleData.Reserve(NumTransforms);
 
-		TArray<TGeometryParticleHandle<Chaos::FReal, 3>*> RelatedBreaks;
+		TArray<int32> RelatedBreaks;
 		RelatedBreaks.Reserve(Breaks.Num());
 		for(const FBreakingData& Break : Breaks)
 		{
-			// Accessing the GT particle here to pull the proxy - while unsafe we're recording a proxy currently so it should remain valid.
-			// No GT data is being read from the particle
-			Chaos::FGeometryParticle* GTParticleUnsafe = Break.Particle->GTGeometryParticle();
-			IPhysicsProxyBase*        BaseProxy        = GTParticleUnsafe->GetProxy();
-			if(BaseProxy->GetType() == EPhysicsProxyType::GeometryCollectionType)
+			// Is the proxy pending destruction? If they are no longer tracked by the PhysScene, the proxy is deleted or pending deletion.
+			if (Scene.GetOwningComponent<UPrimitiveComponent>(Break.Proxy) == nullptr)
+			{
+				continue;
+			}
+			
+			if(Break.Proxy->GetType() == EPhysicsProxyType::GeometryCollectionType)
 			{
-				FGeometryCollectionPhysicsProxy* ConcreteProxy = static_cast<FGeometryCollectionPhysicsProxy*>(BaseProxy);
+				FGeometryCollectionPhysicsProxy* ConcreteProxy = static_cast<FGeometryCollectionPhysicsProxy*>(Break.Proxy);
 
 				if(ConcreteProxy == Proxy)
 				{
 					// The break particle belongs to our proxy
-					RelatedBreaks.Add(Break.Particle);
+					RelatedBreaks.Add(Break.TransformGroupIndex);
 				}
 			}
 		}
@@ -123,8 +129,7 @@ namespace Chaos
 					Pending.PendingTransform = FTransform(Handle->R(), Handle->X()).GetRelativeTransform(InRootTransform);
 				}
 
-				int32 BreakIndex = RelatedBreaks.Find(Handle);
-				if(BreakIndex != INDEX_NONE)
+				if (RelatedBreaks.Contains(TransformIndex))
 				{
 					OutFrame.PushEvent(FEnableStateEvent::EventName, InTime, FEnableStateEvent(TransformIndex, true));
 				}
@@ -138,13 +143,9 @@ namespace Chaos
 				if (BreakingDataArray->IsValidIndex(Index))
 				{
 					const FBreakingData& BreakingData = (*BreakingDataArray)[Index];
-					if (const TPBDRigidParticleHandle<FReal, 3>* Rigid = BreakingData.Particle->CastToRigidParticle())
+					if (BreakingData.TransformGroupIndex > INDEX_NONE)
 					{
-						int32 TransformIndex = Proxy->GetTransformGroupIndexFromHandle(Rigid);
-						if (TransformIndex > INDEX_NONE)
-						{
-							OutFrame.PushEvent(FBreakingEvent::EventName, InTime, FBreakingEvent(TransformIndex, BreakingData, WorldToComponent));
-						}
+						OutFrame.PushEvent(FBreakingEvent::EventName, InTime, FBreakingEvent(BreakingData.TransformGroupIndex, BreakingData, WorldToComponent));
 					}
 				}
 				
@@ -158,14 +159,7 @@ namespace Chaos
 				if (CollisionDataArray->IsValidIndex(Index))
 				{
 					const FCollidingData& CollisionData = (*CollisionDataArray)[Index];
-					if (const TPBDRigidParticleHandle<FReal, 3>* Rigid = CollisionData.Levelset->CastToRigidParticle())
-					{
-						int32 TransformIndex = Proxy->GetTransformGroupIndexFromHandle(Rigid);
-						if (TransformIndex > INDEX_NONE)
-						{
-							OutFrame.PushEvent(FCollisionEvent::EventName, InTime, FCollisionEvent(TransformIndex, CollisionData, WorldToComponent));
-						}
-					}
+					OutFrame.PushEvent(FCollisionEvent::EventName, InTime, FCollisionEvent(CollisionData, WorldToComponent));
 				}
 
 			}
@@ -178,13 +172,9 @@ namespace Chaos
 				if (TrailingDataArray->IsValidIndex(Index))
 				{
 					const FTrailingData& TrailingData = (*TrailingDataArray)[Index];
-					if (const TPBDRigidParticleHandle<FReal, 3>*Rigid = TrailingData.Particle->CastToRigidParticle())
+					if (TrailingData.TransformGroupIndex > INDEX_NONE)
 					{
-						int32 TransformIndex = Proxy->GetTransformGroupIndexFromHandle(Rigid);
-						if (TransformIndex > INDEX_NONE)
-						{
-							OutFrame.PushEvent(FTrailingEvent::EventName, InTime, FTrailingEvent(TransformIndex, TrailingData, WorldToComponent));
-						}
+						OutFrame.PushEvent(FTrailingEvent::EventName, InTime, FTrailingEvent(TrailingData.TransformGroupIndex, TrailingData, WorldToComponent));
 					}
 				}
 
@@ -201,6 +191,8 @@ namespace Chaos
 															FPlaybackTickRecord&                               TickRecord,
 															TArray<TPBDRigidParticleHandle<Chaos::FReal, 3>*>& OutUpdatedRigids) const
 	{
+		ensure(IsInGameThreadContext());
+		
 		using FClusterParticle = Chaos::FPBDRigidClusteredParticleHandle;
 		using FRigidParticle = Chaos::FPBDRigidParticleHandle;
 
@@ -329,7 +321,7 @@ namespace Chaos
 							}
 
 							FBreakingData CachedBreak;
-							CachedBreak.Particle = Particle;
+							CachedBreak.Proxy = Proxy;
 							CachedBreak.Location = ComponentToWorld.TransformPosition(Event->Location);
 							CachedBreak.Velocity = ComponentToWorld.TransformVector(Event->Velocity);
 							CachedBreak.AngularVelocity = Event->AngularVelocity;
@@ -377,7 +369,7 @@ namespace Chaos
 							}
 
 							FTrailingData CachedTrail;
-							CachedTrail.Particle = Particle;
+							CachedTrail.Proxy = Proxy;
 							CachedTrail.Location = ComponentToWorld.TransformPosition(Event->Location);
 							CachedTrail.Velocity = ComponentToWorld.TransformVector(Event->Velocity);
 							CachedTrail.AngularVelocity = Event->AngularVelocity;
@@ -411,53 +403,38 @@ namespace Chaos
 			{
 				if (FCollisionEvent* Event = Handle.Get<FCollisionEvent>())
 				{
-					if (Particles.IsValidIndex(Event->Index))
+					
+					FCollidingData CachedCollision;
+					CachedCollision.Location = ComponentToWorld.TransformPosition(Event->Location);
+					CachedCollision.AccumulatedImpulse = ComponentToWorld.TransformVector(Event->AccumulatedImpulse);
+					CachedCollision.Normal = ComponentToWorld.TransformVector(Event->Normal);
+					CachedCollision.Velocity1 = ComponentToWorld.TransformVector(Event->Velocity1);
+					CachedCollision.Velocity2 = ComponentToWorld.TransformVector(Event->Velocity2);
+					CachedCollision.DeltaVelocity1 = ComponentToWorld.TransformVector(Event->DeltaVelocity1);
+					CachedCollision.DeltaVelocity2 = ComponentToWorld.TransformVector(Event->DeltaVelocity2);
+					CachedCollision.AngularVelocity1 = Event->AngularVelocity1;
+					CachedCollision.AngularVelocity2 = Event->AngularVelocity2;
+					CachedCollision.Mass1 = Event->Mass1;
+					CachedCollision.Mass2 = Event->Mass2;
+					CachedCollision.PenetrationDepth = Event->PenetrationDepth;
+					
+					if (!SolverCollisionEventFilter->Enabled() || SolverCollisionEventFilter->Pass(CachedCollision))
 					{
-						Chaos::FPBDRigidClusteredParticleHandle* Particle = Particles[Event->Index];
-
-						if (Particle)
-						{
-							if (Particle->ObjectState() != EObjectStateType::Kinematic)
+						float TimeStamp = Solver->GetSolverTime();
+						Solver->GetEventManager()->AddEvent<FCollisionEventData>(EEventType::Collision, [&CachedCollision, TimeStamp, Proxy](FCollisionEventData& CollisionEventData)
 							{
-								// If a field or other external actor set the particle to static or dynamic we no longer apply the cache
-								continue;
-							}
-
-							FCollidingData CachedCollision;
-							CachedCollision.Location = ComponentToWorld.TransformPosition(Event->Location);
-							CachedCollision.AccumulatedImpulse = ComponentToWorld.TransformVector(Event->AccumulatedImpulse);
-							CachedCollision.Normal = ComponentToWorld.TransformVector(Event->Normal);
-							CachedCollision.Velocity1 = ComponentToWorld.TransformVector(Event->Velocity1);
-							CachedCollision.Velocity2 = ComponentToWorld.TransformVector(Event->Velocity2);
-							CachedCollision.DeltaVelocity1 = ComponentToWorld.TransformVector(Event->DeltaVelocity1);
-							CachedCollision.DeltaVelocity2 = ComponentToWorld.TransformVector(Event->DeltaVelocity2);
-							CachedCollision.AngularVelocity1 = Event->AngularVelocity1;
-							CachedCollision.AngularVelocity2 = Event->AngularVelocity2;
-							CachedCollision.Mass1 = Event->Mass1;
-							CachedCollision.Mass2 = Event->Mass2;
-							CachedCollision.PenetrationDepth = Event->PenetrationDepth;
-							CachedCollision.Particle = Particle;
-
-							// #todo: Are these even available from a cache?
-							CachedCollision.Levelset = nullptr;
-
-							if (!SolverCollisionEventFilter->Enabled() || SolverCollisionEventFilter->Pass(CachedCollision))
-							{
-								float TimeStamp = Solver->GetSolverTime();
-								Solver->GetEventManager()->AddEvent<FCollisionEventData>(EEventType::Collision, [&CachedCollision, TimeStamp, Particle](FCollisionEventData& CollisionEventData)
-									{
-										if (CollisionEventData.CollisionData.TimeCreated != TimeStamp)
-										{
-											CollisionEventData.CollisionData.AllCollisionsArray.Reset();
-											CollisionEventData.PhysicsProxyToCollisionIndices.Reset();
-											CollisionEventData.CollisionData.TimeCreated = TimeStamp;
-										}
-										int32 NewIdx = CollisionEventData.CollisionData.AllCollisionsArray.Add(CachedCollision);
-										CollisionEventData.PhysicsProxyToCollisionIndices.PhysicsProxyToIndicesMap.FindOrAdd(Particle->PhysicsProxy()).Add(NewIdx);
-									});
-							}
-						}
+								if (CollisionEventData.CollisionData.TimeCreated != TimeStamp)
+								{
+									CollisionEventData.CollisionData.AllCollisionsArray.Reset();
+									CollisionEventData.PhysicsProxyToCollisionIndices.Reset();
+									CollisionEventData.CollisionData.TimeCreated = TimeStamp;
+								}
+								int32 NewIdx = CollisionEventData.CollisionData.AllCollisionsArray.Add(CachedCollision);
+								CollisionEventData.PhysicsProxyToCollisionIndices.PhysicsProxyToIndicesMap.FindOrAdd(Proxy).Add(NewIdx);
+							});
 					}
+						
+					
 				}
 			}
 		}
@@ -584,6 +561,8 @@ namespace Chaos
 	
 	bool FGeometryCollectionCacheAdapter::InitializeForRecord(UPrimitiveComponent* InComponent, UChaosCache* InCache)
 	{
+		ensure(IsInGameThreadContext());
+		
 		UGeometryCollectionComponent*    Comp     = CastChecked<UGeometryCollectionComponent>(InComponent);
 		FGeometryCollectionPhysicsProxy* Proxy    = Comp->GetPhysicsProxy();
 
@@ -627,6 +606,8 @@ namespace Chaos
 
 	bool FGeometryCollectionCacheAdapter::InitializeForPlayback(UPrimitiveComponent* InComponent, UChaosCache* InCache) const
 	{
+		ensure(IsInGameThreadContext());
+		
 		UGeometryCollectionComponent*    Comp  = CastChecked<UGeometryCollectionComponent>(InComponent);
 		if (!Comp)
 		{
@@ -651,6 +632,8 @@ namespace Chaos
 
 	void FGeometryCollectionCacheAdapter::HandleBreakingEvents(const Chaos::FBreakingEventData& Event)
 	{
+		ensure(IsInGameThreadContext());
+		
 		BreakingDataArray = &Event.BreakingData.AllBreakingsArray;
 
 		for (TPair<IPhysicsProxyBase*, FCachedEventData>& Data : CachedData)
@@ -668,6 +651,8 @@ namespace Chaos
 
 	void FGeometryCollectionCacheAdapter::HandleCollisionEvents(const Chaos::FCollisionEventData& Event)
 	{
+		ensure(IsInGameThreadContext());
+		
 		CollisionDataArray = &Event.CollisionData.AllCollisionsArray;
 
 		for (TPair<IPhysicsProxyBase*, FCachedEventData>& Data : CachedData)
@@ -685,6 +670,8 @@ namespace Chaos
 
 	void FGeometryCollectionCacheAdapter::HandleTrailingEvents(const Chaos::FTrailingEventData& Event)
 	{
+		ensure(IsInGameThreadContext());
+		
 		TrailingDataArray = &Event.TrailingData.AllTrailingsArray;
 
 		for (TPair<IPhysicsProxyBase*, FCachedEventData>& Data : CachedData)
diff --git a/Engine/Plugins/Experimental/ChaosCaching/Source/ChaosCaching/Public/Chaos/Adapters/GeometryCollectionComponentCacheAdapter.h b/Engine/Plugins/Experimental/ChaosCaching/Source/ChaosCaching/Public/Chaos/Adapters/GeometryCollectionComponentCacheAdapter.h
--- a/Engine/Plugins/Experimental/ChaosCaching/Source/ChaosCaching/Public/Chaos/Adapters/GeometryCollectionComponentCacheAdapter.h
+++ b/Engine/Plugins/Experimental/ChaosCaching/Source/ChaosCaching/Public/Chaos/Adapters/GeometryCollectionComponentCacheAdapter.h
@@ -94,8 +94,7 @@ struct FCollisionEvent : public FCacheEventBase
 	static FName EventName;
 
 	FCollisionEvent()
-		: Index(INDEX_NONE)
-		, Location(0.0f, 0.0f, 0.0f)
+		: Location(0.0f, 0.0f, 0.0f)
 		, AccumulatedImpulse(0.0f, 0.0f, 0.0f)
 		, Normal(0.0f, 0.0f, 1.0f)
 		, Velocity1(0.0f, 0.0f, 0.0f)
@@ -109,9 +108,8 @@ struct FCollisionEvent : public FCacheEventBase
 		, PenetrationDepth(0.0f)
 	{}
 
-	FCollisionEvent(int32 InIndex, const Chaos::FCollidingData& InData, const FTransform& WorldToComponent)
-		: Index(InIndex)
-		, Location(WorldToComponent.TransformPosition(InData.Location))
+	FCollisionEvent(const Chaos::FCollidingData& InData, const FTransform& WorldToComponent)
+		: Location(WorldToComponent.TransformPosition(InData.Location))
 		, AccumulatedImpulse(WorldToComponent.TransformVector(InData.AccumulatedImpulse))
 		, Normal(WorldToComponent.TransformVector(InData.Normal))
 		, Velocity1(WorldToComponent.TransformVector(InData.Velocity1))
@@ -124,10 +122,7 @@ struct FCollisionEvent : public FCacheEventBase
 		, Mass2(InData.Mass2)
 		, PenetrationDepth(InData.PenetrationDepth)
 	{}
-
-	UPROPERTY()
-	int32 Index; 
-		
+	
 	UPROPERTY()
 	FVector Location;
 	
diff --git a/Engine/Plugins/Experimental/ChaosNiagara/Source/ChaosNiagara/Private/NiagaraDataInterfaceChaosDestruction.cpp b/Engine/Plugins/Experimental/ChaosNiagara/Source/ChaosNiagara/Private/NiagaraDataInterfaceChaosDestruction.cpp
--- a/Engine/Plugins/Experimental/ChaosNiagara/Source/ChaosNiagara/Private/NiagaraDataInterfaceChaosDestruction.cpp
+++ b/Engine/Plugins/Experimental/ChaosNiagara/Source/ChaosNiagara/Private/NiagaraDataInterfaceChaosDestruction.cpp
@@ -1254,15 +1254,19 @@ void UNiagaraDataInterfaceChaosDestruction::HandleBreakingEvents(const Chaos::FB
 
 			UPhysicalMaterial* PhysicalMaterial = nullptr;
 			int32 MaterialID = 0;
+			
+			// #todo (bmiller) Need to have a better way to get at this data now that particle handle isn't in the Breaking event payload.
+			/*
 			if(DataIn.Particle && DataIn.Particle->Geometry() )
 				MaterialID = DataIn.Particle->Geometry()->GetMaterialIndex(0);
+			*/
 
 			UGeometryCollectionComponent* GeometryCollectionComponent = nullptr;
 			UMaterialInterface* Material = nullptr;
 #if INCLUDE_CHAOS
 			for (auto& Solver : Solvers)
 			{
-				GeometryCollectionComponent = Solver.PhysScene->GetOwningComponent<UGeometryCollectionComponent>(DataIn.Particle->PhysicsProxy());
+				GeometryCollectionComponent = Solver.PhysScene->GetOwningComponent<UGeometryCollectionComponent>(DataIn.Proxy);
 				if (GeometryCollectionComponent)
 					break;
 			}
diff --git a/Engine/Source/Runtime/Engine/Private/PhysicsEngine/Experimental/PhysScene_Chaos.cpp b/Engine/Source/Runtime/Engine/Private/PhysicsEngine/Experimental/PhysScene_Chaos.cpp
--- a/Engine/Source/Runtime/Engine/Private/PhysicsEngine/Experimental/PhysScene_Chaos.cpp
+++ b/Engine/Source/Runtime/Engine/Private/PhysicsEngine/Experimental/PhysScene_Chaos.cpp
@@ -706,7 +706,13 @@ void FPhysScene_Chaos::HandleCollisionEvents(const Chaos::FCollisionEventData& E
 							int32 CollisionIdx = Chaos::FEventManager::DecodeCollisionIndex(EncodedCollisionIdx, bSwapOrder);
 
 							Chaos::FCollidingData const& CollisionDataItem = CollisionData[CollisionIdx];
-							IPhysicsProxyBase* const PhysicsProxy1 = bSwapOrder ? CollisionDataItem.Particle->PhysicsProxy() : CollisionDataItem.Levelset->PhysicsProxy();
+							IPhysicsProxyBase* const PhysicsProxy1 = bSwapOrder ? CollisionDataItem.Proxy1 : CollisionDataItem.Proxy2;
+
+							// Are the proxies pending destruction? If they are no longer tracked by the PhysScene, the proxy is deleted or pending deletion.
+							if (GetOwningComponent<UPrimitiveComponent>(PhysicsProxy0) == nullptr || GetOwningComponent<UPrimitiveComponent>(PhysicsProxy1) == nullptr)
+							{
+								continue;
+							}
 
 							bool bNewEntry = false;
 							FCollisionNotifyInfo& NotifyInfo = GetPendingCollisionForContactPair(PhysicsProxy0, PhysicsProxy1, bNewEntry);
diff --git a/Engine/Source/Runtime/Experimental/Chaos/Private/Chaos/PBDRigidClustering.cpp b/Engine/Source/Runtime/Experimental/Chaos/Private/Chaos/PBDRigidClustering.cpp
--- a/Engine/Source/Runtime/Experimental/Chaos/Private/Chaos/PBDRigidClustering.cpp
+++ b/Engine/Source/Runtime/Experimental/Chaos/Private/Chaos/PBDRigidClustering.cpp
@@ -21,6 +21,7 @@
 #include "Chaos/PerParticleEulerStepVelocity.h"
 #include "Chaos/PerParticleEtherDrag.h"
 #include "Chaos/PerParticlePBDEulerStep.h"
+#include "PhysicsProxy/GeometryCollectionPhysicsProxy.h"
 #include "CoreMinimal.h"
 
 namespace Chaos
@@ -797,11 +798,24 @@ namespace Chaos
 					{
 						const int32 NewIdx = MAllClusterBreakings.Add(FBreakingData());
 						FBreakingData& ClusterBreak = MAllClusterBreakings[NewIdx];
-						ClusterBreak.Particle = Child;
+						ClusterBreak.Proxy = Child->PhysicsProxy();
 						ClusterBreak.Location = Child->X();
 						ClusterBreak.Velocity = Child->V();
 						ClusterBreak.AngularVelocity = Child->W();
 						ClusterBreak.Mass = Child->M();
+						if (Child->Geometry()->HasBoundingBox())
+						{
+							ClusterBreak.BoundingBox = Child->Geometry()->BoundingBox();
+						}
+						if (ClusterBreak.Proxy->GetType() == EPhysicsProxyType::GeometryCollectionType)
+						{
+							FGeometryCollectionPhysicsProxy* ConcreteProxy = static_cast<FGeometryCollectionPhysicsProxy*>(ClusterBreak.Proxy);
+							ClusterBreak.TransformGroupIndex = ConcreteProxy->GetTransformGroupIndexFromHandle(Child);
+						}
+						else
+						{
+							ClusterBreak.TransformGroupIndex = INDEX_NONE;
+						}
 					}
 				}
 			}
@@ -1065,11 +1079,24 @@ namespace Chaos
 					{
 						const int32 NewIdx = MAllClusterBreakings.Add(FBreakingData());
 						FBreakingData& ClusterBreak = MAllClusterBreakings[NewIdx];
-						ClusterBreak.Particle = Child;
+						ClusterBreak.Proxy = Child->PhysicsProxy();
 						ClusterBreak.Location = Child->X();
 						ClusterBreak.Velocity = Child->V();
 						ClusterBreak.AngularVelocity = Child->W();
 						ClusterBreak.Mass = Child->M();
+						if (Child->Geometry()->HasBoundingBox())
+						{
+							ClusterBreak.BoundingBox = Child->Geometry()->BoundingBox();
+						}
+						if (ClusterBreak.Proxy->GetType() == EPhysicsProxyType::GeometryCollectionType)
+						{
+							FGeometryCollectionPhysicsProxy* ConcreteProxy = static_cast<FGeometryCollectionPhysicsProxy*>(ClusterBreak.Proxy);
+							ClusterBreak.TransformGroupIndex = ConcreteProxy->GetTransformGroupIndexFromHandle(Child);
+						}
+						else
+						{
+							ClusterBreak.TransformGroupIndex = INDEX_NONE;
+						}
 					}
 				}
 			}
@@ -1380,11 +1407,24 @@ namespace Chaos
 						{
 							int32 NewIdx = MAllClusterBreakings.Add(FBreakingData());
 							FBreakingData& ClusterBreak = MAllClusterBreakings[NewIdx];
-							ClusterBreak.Particle = ClusteredParticle;
+							ClusterBreak.Proxy = ClusteredParticle->PhysicsProxy();
 							ClusterBreak.Location = ClusteredParticle->X();
 							ClusterBreak.Velocity = ClusteredParticle->V();
 							ClusterBreak.AngularVelocity = ClusteredParticle->W();
 							ClusterBreak.Mass = ClusteredParticle->M();
+							if (ClusteredParticle->Geometry()->HasBoundingBox())
+							{
+								ClusterBreak.BoundingBox = ClusteredParticle->Geometry()->BoundingBox();
+							}
+							if (ClusterBreak.Proxy->GetType() == EPhysicsProxyType::GeometryCollectionType)
+							{
+								FGeometryCollectionPhysicsProxy* ConcreteProxy = static_cast<FGeometryCollectionPhysicsProxy*>(ClusterBreak.Proxy);
+								ClusterBreak.TransformGroupIndex = ConcreteProxy->GetTransformGroupIndexFromHandle(ClusteredParticle);
+							}
+							else
+							{
+								ClusterBreak.TransformGroupIndex = INDEX_NONE;
+							}
 						}
 					}
 				}
@@ -2007,11 +2047,25 @@ namespace Chaos
 		{
 			const int32 NewIdx = MAllClusterBreakings.Add(FBreakingData());
 			FBreakingData& ClusterBreak = MAllClusterBreakings[NewIdx];
-			ClusterBreak.Particle = Particle;
+			ClusterBreak.Proxy = Particle->PhysicsProxy();
 			ClusterBreak.Location = Particle->X();
 			ClusterBreak.Velocity = Particle->V();
 			ClusterBreak.AngularVelocity = Particle->W();
 			ClusterBreak.Mass = Particle->M();
+			if (Particle->Geometry()->HasBoundingBox())
+			{
+				ClusterBreak.BoundingBox = Particle->Geometry()->BoundingBox();
+			}
+			if (ClusterBreak.Proxy->GetType() == EPhysicsProxyType::GeometryCollectionType)
+			{
+				FGeometryCollectionPhysicsProxy* ConcreteProxy = static_cast<FGeometryCollectionPhysicsProxy*>(ClusterBreak.Proxy);
+				ClusterBreak.TransformGroupIndex = ConcreteProxy->GetTransformGroupIndexFromHandle(Particle);
+			}
+			else
+			{
+				ClusterBreak.TransformGroupIndex = INDEX_NONE;
+			}
+			
 		}
 	}
 
diff --git a/Engine/Source/Runtime/Experimental/Chaos/Private/EventDefaults.cpp b/Engine/Source/Runtime/Experimental/Chaos/Private/EventDefaults.cpp
--- a/Engine/Source/Runtime/Experimental/Chaos/Private/EventDefaults.cpp
+++ b/Engine/Source/Runtime/Experimental/Chaos/Private/EventDefaults.cpp
@@ -31,6 +31,8 @@ namespace Chaos
 		(const Chaos::FPBDRigidsSolver* Solver, FCollisionEventData& CollisionEventData)
 		{
 			check(Solver);
+			ensure(IsInPhysicsThreadContext());
+
 			SCOPE_CYCLE_COUNTER(STAT_GatherCollisionEvent);
 
 			// #todo: This isn't working - SolverActor parameters are set on a solver but it is currently a different solver that is simulating!!
@@ -145,8 +147,15 @@ namespace Chaos
 							Data.AccumulatedImpulse = Constraint.AccumulatedImpulse;
 							Data.Normal = Constraint.GetNormal();
 							Data.PenetrationDepth = Constraint.GetPhi();
-							Data.Particle = Particle0;
-							Data.Levelset = Particle1;
+							
+							Data.Proxy1 = Particle0 ? Particle0->PhysicsProxy() : nullptr;
+							Data.Proxy2 = Particle1 ? Particle1->PhysicsProxy() : nullptr;
+
+							// Collision constraints require both proxies are valid. If either is not, we needn't record the collision event.
+							if (Data.Proxy1 == nullptr || Data.Proxy2 == nullptr)
+							{
+								continue;
+							}
 
 							if (FPBDRigidParticleHandle * Rigid0 = Particle0->CastToRigidParticle())
 							{
@@ -177,9 +186,7 @@ namespace Chaos
 
 							IPhysicsProxyBase* const PhysicsProxy = Particle0->PhysicsProxy();
 							IPhysicsProxyBase* const OtherPhysicsProxy = Particle1->PhysicsProxy();
-							//Data.Material1 = nullptr; // #todo: provide UPhysicalMaterial for Particle
-							//Data.Material2 = nullptr; // #todo: provide UPhysicalMaterial for Levelset
-
+							
 							const FSolverCollisionEventFilter* SolverCollisionEventFilter = Solver->GetEventFilters()->GetCollisionFilter();
 							if (!SolverCollisionEventFilter->Enabled() || SolverCollisionEventFilter->Pass(Data))
 
@@ -227,6 +234,8 @@ namespace Chaos
 		(const Chaos::FPBDRigidsSolver* Solver, FBreakingEventData& BreakingEventData)
 		{
 			check(Solver);
+			ensure(IsInPhysicsThreadContext());
+
 			SCOPE_CYCLE_COUNTER(STAT_GatherBreakingEvent);
 
 			// #todo: This isn't working - SolverActor parameters are set on a solver but it is currently a different solver that is simulating!!
@@ -251,54 +260,40 @@ namespace Chaos
 #if TODO_REIMPLEMENT_RIGID_CLUSTERING
 			const Chaos::FPBDRigidsSolver::FClusteringType::FClusterMap& ParentToChildrenMap = Evolution->GetRigidClustering().GetChildrenMap();
 #endif
-
-			if(AllBreakingsArray.Num() > 0)
+			
+			if (AllBreakingsArray.Num() > 0)
 			{
-				for(int32 Idx = 0; Idx < AllBreakingsArray.Num(); ++Idx)
-				{
-					// Since Clustered GCs can be unioned the particleIndex representing the union 
-					// is not associated with a PhysicsProxy
-					FPBDRigidParticleHandle* PBDRigid = AllBreakingsArray[Idx].Particle->CastToRigidParticle();
-					if(PBDRigid)
+				for (int32 Idx = 0; Idx < AllBreakingsArray.Num(); ++Idx)
+				{					
+					FBreakingData BreakingData;
+					BreakingData.Location = AllBreakingsArray[Idx].Location;
+					BreakingData.Velocity = AllBreakingsArray[Idx].Velocity;
+					BreakingData.AngularVelocity = AllBreakingsArray[Idx].AngularVelocity;
+					BreakingData.Mass = AllBreakingsArray[Idx].Mass;
+					BreakingData.Proxy = AllBreakingsArray[Idx].Proxy;
+					BreakingData.BoundingBox = AllBreakingsArray[Idx].BoundingBox;
+					BreakingData.TransformGroupIndex = AllBreakingsArray[Idx].TransformGroupIndex;
+
+					const FSolverBreakingEventFilter* SolverBreakingEventFilter = Solver->GetEventFilters()->GetBreakingFilter();
+					if (!SolverBreakingEventFilter->Enabled() || SolverBreakingEventFilter->Pass(BreakingData))
 					{
-						if(ensure(!AllBreakingsArray[Idx].Location.ContainsNaN() &&
-							!PBDRigid->V().ContainsNaN() &&
-							!PBDRigid->W().ContainsNaN()))
-						{
-							FBreakingData BreakingData;
-							BreakingData.Location = AllBreakingsArray[Idx].Location;
-							BreakingData.Velocity = PBDRigid->V();
-							BreakingData.AngularVelocity = PBDRigid->W();
-							BreakingData.Mass = PBDRigid->M();
-							BreakingData.Particle = PBDRigid;
-							
-							if(PBDRigid->Geometry()->HasBoundingBox())
-							{
-								BreakingData.BoundingBox = PBDRigid->Geometry()->BoundingBox();
-							}
+						int32 NewIdx = AllBreakingDataArray.Add(FBreakingData());
+						FBreakingData& BreakingDataArrayItem = AllBreakingDataArray[NewIdx];
+						BreakingDataArrayItem = BreakingData;
 
-							const FSolverBreakingEventFilter* SolverBreakingEventFilter = Solver->GetEventFilters()->GetBreakingFilter();
-							if(!SolverBreakingEventFilter->Enabled() || SolverBreakingEventFilter->Pass(BreakingData))
-							{
-								int32 NewIdx = AllBreakingDataArray.Add(FBreakingData());
-								FBreakingData& BreakingDataArrayItem = AllBreakingDataArray[NewIdx];
-								BreakingDataArrayItem = BreakingData;
-
-								// Add to AllBreakingIndicesByPhysicsProxy
-								AllBreakingIndicesByPhysicsProxy.FindOrAdd(BreakingData.Particle->PhysicsProxy()).Add(FEventManager::EncodeCollisionIndex(NewIdx, false));
+						// Add to AllBreakingIndicesByPhysicsProxy
+						AllBreakingIndicesByPhysicsProxy.FindOrAdd(BreakingData.Proxy).Add(FEventManager::EncodeCollisionIndex(NewIdx, false));
 
 #if 0 // #todo
-								// If AllBreakingsArray[Idx].ParticleIndex is a cluster store an index for a mesh in this cluster
-								if(ClusterIdsArray[AllBreakingsArray[Idx].ParticleIndex].NumChildren > 0)
-								{
-									int32 ParticleIndexMesh = GetParticleIndexMesh(ParentToChildrenMap, AllBreakingsArray[Idx].ParticleIndex);
-									ensure(ParticleIndexMesh != INDEX_NONE);
-									BreakingDataArrayItem.ParticleIndexMesh = ParticleIndexMesh;
-								}
-#endif
-								}
+						// If AllBreakingsArray[Idx].ParticleIndex is a cluster store an index for a mesh in this cluster
+						if (ClusterIdsArray[AllBreakingsArray[Idx].ParticleIndex].NumChildren > 0)
+						{
+							int32 ParticleIndexMesh = GetParticleIndexMesh(ParentToChildrenMap, AllBreakingsArray[Idx].ParticleIndex);
+							ensure(ParticleIndexMesh != INDEX_NONE);
+							BreakingDataArrayItem.ParticleIndexMesh = ParticleIndexMesh;
 						}
-					}
+#endif
+					}					
 				}
 			}
 		
@@ -311,6 +306,7 @@ namespace Chaos
 		(const Chaos::FPBDRigidsSolver* Solver, FTrailingEventData& TrailingEventData)
 		{
 			check(Solver);
+			ensure(IsInPhysicsThreadContext());
 
 			// #todo: This isn't working - SolverActor parameters are set on a solver but it is currently a different solver that is simulating!!
 			if (!Solver->GetEventFilters()->IsTrailingEventEnabled())
@@ -355,14 +351,23 @@ namespace Chaos
 							TrailingData.Velocity = ActiveParticle->V();
 							TrailingData.AngularVelocity = ActiveParticle->W();
 							TrailingData.Mass = ActiveParticle->M();
-
-							TrailingData.Particle = ActiveParticle;
+							TrailingData.Proxy = ActiveParticle->PhysicsProxy();
 							
 							if (ActiveParticle->Geometry()->HasBoundingBox())
 							{
 								TrailingData.BoundingBox = ActiveParticle->Geometry()->BoundingBox();
 							}
 
+							if (TrailingData.Proxy->GetType() == EPhysicsProxyType::GeometryCollectionType)
+							{
+								FGeometryCollectionPhysicsProxy* ConcreteProxy = static_cast<FGeometryCollectionPhysicsProxy*>(TrailingData.Proxy);
+								TrailingData.TransformGroupIndex = ConcreteProxy->GetTransformGroupIndexFromHandle(ActiveParticle);
+							}
+							else
+							{
+								TrailingData.TransformGroupIndex = INDEX_NONE;
+							}
+
 							const FSolverTrailingEventFilter* SolverTrailingEventFilter = Solver->GetEventFilters()->GetTrailingFilter();
 							if (!SolverTrailingEventFilter->Enabled() || SolverTrailingEventFilter->Pass(TrailingData))
 							{
@@ -371,7 +376,7 @@ namespace Chaos
 								TrailingDataArrayItem = TrailingData;
 
 								// Add to AllTrailingIndicesByPhysicsProxy
-								AllTrailingIndicesByPhysicsProxy.FindOrAdd(TrailingData.Particle->PhysicsProxy()).Add(FEventManager::EncodeCollisionIndex(NewIdx, false));
+								AllTrailingIndicesByPhysicsProxy.FindOrAdd(TrailingData.Proxy).Add(FEventManager::EncodeCollisionIndex(NewIdx, false));
 
 								// If IdxParticle is a cluster store an index for a mesh in this cluster
 #if 0
@@ -396,6 +401,8 @@ namespace Chaos
 		(const Chaos::FPBDRigidsSolver* Solver, FSleepingEventData& SleepingEventData)
 		{
 			check(Solver);
+			ensure(IsInPhysicsThreadContext());
+
 			SCOPE_CYCLE_COUNTER(STAT_GatherSleepingEvent);
 
 			const auto* Evolution = Solver->GetEvolution();
@@ -416,7 +423,7 @@ namespace Chaos
 					{
 						int32 NewIdx = EventSleepDataArray.Add(FSleepingData());
 						FSleepingData& SleepingDataArrayItem = EventSleepDataArray[NewIdx];
-						SleepingDataArrayItem.Particle = Particle;
+						SleepingDataArrayItem.Proxy = SleepData.Particle->PhysicsProxy();
 						SleepingDataArrayItem.Sleeping = SleepData.Sleeping;
 					}
 				}
diff --git a/Engine/Source/Runtime/Experimental/Chaos/Public/Chaos/ExternalCollisionData.h b/Engine/Source/Runtime/Experimental/Chaos/Public/Chaos/ExternalCollisionData.h
--- a/Engine/Source/Runtime/Experimental/Chaos/Public/Chaos/ExternalCollisionData.h
+++ b/Engine/Source/Runtime/Experimental/Chaos/Public/Chaos/ExternalCollisionData.h
@@ -29,13 +29,12 @@ namespace Chaos
 			, Mass1((FReal)0.0)
 			, Mass2((FReal)0.0)
 			, PenetrationDepth((FReal)0.0)
-			, Particle(nullptr)
-			, Levelset(nullptr)
+			, Proxy1(nullptr)
+			, Proxy2(nullptr)
 		{}
 
 		FCollidingData(FVec3 InLocation, FVec3 InAccumulatedImpulse, FVec3 InNormal, FVec3 InVelocity1, FVec3 InVelocity2, FVec3 InDeltaVelocity1, FVec3 InDeltaVelocity2
-			, FVec3 InAngularVelocity1, FVec3 InAngularVelocity2, FReal InMass1,FReal InMass2,  FReal InPenetrationDepth, FGeometryParticleHandle* InParticle
-			, FGeometryParticleHandle* InLevelset)
+			, FVec3 InAngularVelocity1, FVec3 InAngularVelocity2, FReal InMass1,FReal InMass2,  FReal InPenetrationDepth, IPhysicsProxyBase* InProxy1, IPhysicsProxyBase* InProxy2)
 			: Location(InLocation)
 			, AccumulatedImpulse(InAccumulatedImpulse)
 			, Normal(InNormal)
@@ -48,8 +47,8 @@ namespace Chaos
 			, Mass1(InMass1)
 			, Mass2(InMass2)
 			, PenetrationDepth(InPenetrationDepth)
-			, Particle(InParticle)
-			, Levelset(InLevelset)
+			, Proxy1(InProxy1)
+			, Proxy2(InProxy2)
 		{}
 
 		FVec3 Location;
@@ -60,8 +59,15 @@ namespace Chaos
 		FVec3 AngularVelocity1, AngularVelocity2;
 		FReal Mass1, Mass2;
 		FReal PenetrationDepth;
-		FGeometryParticleHandle* Particle;
-		FGeometryParticleHandle* Levelset;
+
+		// The pointers to the proxies should be used with caution on the Game Thread.
+		// Ideally we only ever use these as table keys when acquiring related structures.
+		// If we genuinely need to dereference the pointers for any reason, test if they are deleted (nullptr) or
+		// pending deletion: if a call to FPhysScene_Chaos::GetOwningComponent<UPrimitiveComponent>() returns
+		// nullptr, the proxies should not be used.
+		// If either Proxy is nullptr, the structure is invalid.
+		IPhysicsProxyBase* Proxy1;
+		IPhysicsProxyBase* Proxy2;
 	};
 
 	/*
@@ -79,8 +85,6 @@ namespace Chaos
 			, AngularVelocity2(FVec3((FReal)0.0))
 			, Mass1((FReal)0.0)
 			, Mass2((FReal)0.0)
-			, Particle(nullptr)
-			, Levelset(nullptr)
 			, BoundingboxVolume((FReal)-1.0)
 			, BoundingboxExtentMin((FReal)-1.0)
 			, BoundingboxExtentMax((FReal)-1.0)
@@ -89,7 +93,7 @@ namespace Chaos
 
 		FCollidingDataExt(
 		    FVec3 InLocation, FVec3 InAccumulatedImpulse, FVec3 InNormal, FVec3 InVelocity1, FVec3 InVelocity2
-			, FVec3 InAngularVelocity1, FVec3 InAngularVelocity2, FReal InMass1, FReal InMass2, FGeometryParticleHandle* InParticle, FGeometryParticleHandle* InLevelset
+			, FVec3 InAngularVelocity1, FVec3 InAngularVelocity2, FReal InMass1, FReal InMass2 //, FGeometryParticleHandle* InParticle, FGeometryParticleHandle* InLevelset
 			, FReal InBoundingboxVolume, FReal InBoundingboxExtentMin, FReal InBoundingboxExtentMax, int32 InSurfaceType)
 			: Location(InLocation)
 			, AccumulatedImpulse(InAccumulatedImpulse)
@@ -100,8 +104,6 @@ namespace Chaos
 			, AngularVelocity2(InAngularVelocity2)
 			, Mass1(InMass1)
 			, Mass2(InMass2)
-			, Particle(InParticle)
-			, Levelset(InLevelset)
 			, BoundingboxVolume(InBoundingboxVolume)
 			, BoundingboxExtentMin(InBoundingboxExtentMin)
 			, BoundingboxExtentMax(InBoundingboxExtentMax)
@@ -118,8 +120,6 @@ namespace Chaos
 			, AngularVelocity2(InCollisionData.AngularVelocity2)
 			, Mass1(InCollisionData.Mass1)
 			, Mass2(InCollisionData.Mass2)
-			, Particle(InCollisionData.Particle)
-			, Levelset(InCollisionData.Levelset)
 			, BoundingboxVolume((FReal)-1.0)
 			, BoundingboxExtentMin((FReal)-1.0)
 			, BoundingboxExtentMax((FReal)-1.0)
@@ -133,8 +133,6 @@ namespace Chaos
 		FVec3 Velocity1, Velocity2;
 		FVec3 AngularVelocity1, AngularVelocity2;
 		FReal Mass1, Mass2;
-		FGeometryParticleHandle* Particle;
-		FGeometryParticleHandle* Levelset;
 		FReal BoundingboxVolume;
 		FReal BoundingboxExtentMin, BoundingboxExtentMax;
 		int32 SurfaceType;
@@ -146,20 +144,28 @@ namespace Chaos
 	struct FBreakingData
 	{
 		FBreakingData()
-			: Particle(nullptr)
+			: Proxy(nullptr)
 			, Location(FVec3((FReal)0.0))
 			, Velocity(FVec3((FReal)0.0))
 			, AngularVelocity(FVec3((FReal)0.0))
 			, Mass((FReal)0.0)
 			, BoundingBox(FAABB3(FVec3((FReal)0.0), FVec3((FReal)0.0)))
+			, TransformGroupIndex(INDEX_NONE)
 		{}
 
-		FGeometryParticleHandle* Particle;
+		// The pointer to the proxy should be used with caution on the Game Thread.
+		// Ideally we only ever use this as a table key when acquiring related structures.
+		// If we genuinely need to dereference the pointer for any reason, test if it is deleted (nullptr) or
+		// pending deletion: if a call to FPhysScene_Chaos::GetOwningComponent<UPrimitiveComponent>() returns
+		// nullptr, the proxy should not be used.
+		IPhysicsProxyBase* Proxy;
+		
 		FVec3 Location;
 		FVec3 Velocity;
 		FVec3 AngularVelocity;
 		FReal Mass;
 		Chaos::FAABB3 BoundingBox;
+		int32 TransformGroupIndex;
 	};
 
 	/*
@@ -172,7 +178,6 @@ namespace Chaos
 			, Velocity(FVec3((FReal)0.0))
 			, AngularVelocity(FVec3((FReal)0.0))
 			, Mass((FReal)0.0)
-		    , Particle(nullptr)
 			, BoundingboxVolume((FReal)-1.0)
 			, BoundingboxExtentMin((FReal)-1.0)
 			, BoundingboxExtentMax((FReal)-1.0)
@@ -192,7 +197,6 @@ namespace Chaos
 			, Velocity(InVelocity)
 			, AngularVelocity(InAngularVelocity)
 			, Mass(InMass)
-		    , Particle(InParticle)
 			, BoundingboxVolume(InBoundingboxVolume)
 			, BoundingboxExtentMin(InBoundingboxExtentMin)
 			, BoundingboxExtentMax(InBoundingboxExtentMax)
@@ -204,7 +208,6 @@ namespace Chaos
 			, Velocity(InBreakingData.Velocity)
 			, AngularVelocity(InBreakingData.AngularVelocity)
 			, Mass(InBreakingData.Mass)
-		    , Particle(InBreakingData.Particle)
 			, BoundingboxVolume((FReal)-1.0)
 			, BoundingboxExtentMin((FReal)-1.0)
 			, BoundingboxExtentMax((FReal)-1.0)
@@ -216,7 +219,6 @@ namespace Chaos
 		FVec3 Velocity;
 		FVec3 AngularVelocity;
 		FReal Mass;
-		FGeometryParticleHandle* Particle;
 		FReal BoundingboxVolume;
 		FReal BoundingboxExtentMin, BoundingboxExtentMax;
 		int32 SurfaceType;
@@ -243,36 +245,35 @@ namespace Chaos
 			, Velocity(FVec3((FReal)0.0))
 			, AngularVelocity(FVec3((FReal)0.0))
 			, Mass((FReal)0.0)
-			, Particle(nullptr)
+			, Proxy(nullptr)
 			, BoundingBox(FAABB3(FVec3((FReal)0.0), FVec3((FReal)0.0)))
 		{}
 
 		FTrailingData(FVec3 InLocation, FVec3 InVelocity, FVec3 InAngularVelocity, FReal InMass
-			, FGeometryParticleHandle* InParticle, Chaos::TAABB<FReal, 3>& InBoundingBox)
+			, IPhysicsProxyBase* InProxy, Chaos::TAABB<FReal, 3>& InBoundingBox)
 			: Location(InLocation)
 			, Velocity(InVelocity)
 			, AngularVelocity(InAngularVelocity)
 			, Mass(InMass)
-			, Particle(InParticle)
+			, Proxy(InProxy)
 			, BoundingBox(InBoundingBox)
+			, TransformGroupIndex(INDEX_NONE)
 		{}
 
 		FVec3 Location;
 		FVec3 Velocity;
 		FVec3 AngularVelocity;
 		FReal Mass;
-		FGeometryParticleHandle* Particle;
-		Chaos::FAABB3 BoundingBox;
-
-		friend inline uint32 GetTypeHash(const FTrailingData& Other)
-		{
-			return ::GetTypeHash(Other.Particle);
-		}
+		
+		// The pointer to the proxy should be used with caution on the Game Thread.
+		// Ideally we only ever use this as a table key when acquiring related structures.
+		// If we genuinely need to dereference the pointer for any reason, test if it is deleted (nullptr) or
+		// pending deletion: if a call to FPhysScene_Chaos::GetOwningComponent<UPrimitiveComponent>() returns
+		// nullptr, the proxy should not be used.
+		IPhysicsProxyBase* Proxy;
 
-		friend bool operator==(const FTrailingData& A, const FTrailingData& B)
-		{
-			return A.Particle == B.Particle;
-		}
+		Chaos::FAABB3 BoundingBox;
+		int32 TransformGroupIndex;
 	};
 
 	/*
@@ -285,7 +286,6 @@ namespace Chaos
 			, Velocity(FVec3((FReal)0.0))
 			, AngularVelocity(FVec3((FReal)0.0))
 			, Mass((FReal)0.0)
-			, Particle(nullptr)
 			, BoundingboxVolume((FReal)-1.0)
 			, BoundingboxExtentMin((FReal)-1.0)
 			, BoundingboxExtentMax((FReal)-1.0)
@@ -305,7 +305,6 @@ namespace Chaos
 			, Velocity(InVelocity)
 			, AngularVelocity(InAngularVelocity)
 			, Mass(InMass)
-			, Particle(InParticle)
 			, BoundingboxVolume(InBoundingboxVolume)
 			, BoundingboxExtentMin(InBoundingboxExtentMin)
 			, BoundingboxExtentMax(InBoundingboxExtentMax)
@@ -317,7 +316,6 @@ namespace Chaos
 			, Velocity(InTrailingData.Velocity)
 			, AngularVelocity(InTrailingData.AngularVelocity)
 			, Mass(InTrailingData.Mass)
-			, Particle(InTrailingData.Particle)
 			, BoundingboxVolume((FReal)-1.0)
 			, BoundingboxExtentMin((FReal)-1.0)
 			, BoundingboxExtentMax((FReal)-1.0)
@@ -328,38 +326,30 @@ namespace Chaos
 		FVec3 Velocity;
 		FVec3 AngularVelocity;
 		FReal Mass;
-		FGeometryParticleHandle* Particle;
-		//	int32 ParticleIndexMesh; // If ParticleIndex points to a cluster then this index will point to an actual mesh in the cluster
-								 // It is important to be able to get extra data from the component
 		FReal BoundingboxVolume;
 		FReal BoundingboxExtentMin, BoundingboxExtentMax;
 		int32 SurfaceType;
-
-		friend inline uint32 GetTypeHash(const FTrailingDataExt& Other)
-		{
-			return ::GetTypeHash(Other.Particle);
-		}
-
-		friend bool operator==(const FTrailingDataExt& A, const FTrailingDataExt& B)
-		{
-			return A.Particle == B.Particle;
-		}
 	};
 
 	struct FSleepingData
 	{
 		FSleepingData()
-			: Particle(nullptr)
+			: Proxy(nullptr)
 			, Sleeping(true)
 		{}
 
-		FSleepingData(
-		    FGeometryParticle* InParticle, bool InSleeping)
-			: Particle(InParticle)
+		FSleepingData(IPhysicsProxyBase* InProxy, bool InSleeping)
+			: Proxy(InProxy)
 			, Sleeping(InSleeping)
 		{}
 
-		FGeometryParticle* Particle;
+		// The pointer to the proxy should be used with caution on the Game Thread.
+		// Ideally we only ever use this as a table key when acquiring related structures.
+		// If we genuinely need to dereference the pointer for any reason, test if it is deleted (nullptr) or
+		// pending deletion: if a call to FPhysScene_Chaos::GetOwningComponent<UPrimitiveComponent>() returns
+		// nullptr, the proxy should not be used.
+		IPhysicsProxyBase* Proxy;
+
 		bool Sleeping;	// if !Sleeping == Awake
 	};
 
diff --git a/Engine/Source/Runtime/Experimental/ChaosSolverEngine/Private/Chaos/ChaosGameplayEventDispatcher.cpp b/Engine/Source/Runtime/Experimental/ChaosSolverEngine/Private/Chaos/ChaosGameplayEventDispatcher.cpp
--- a/Engine/Source/Runtime/Experimental/ChaosSolverEngine/Private/Chaos/ChaosGameplayEventDispatcher.cpp
+++ b/Engine/Source/Runtime/Experimental/ChaosSolverEngine/Private/Chaos/ChaosGameplayEventDispatcher.cpp
@@ -296,85 +296,86 @@ void UChaosGameplayEventDispatcher::HandleCollisionEvents(const Chaos::FCollisio
 
 								Chaos::FCollidingData const& CollisionDataItem = CollisionData[CollisionIdx];
 
-								// If we have no LevelSet pointer (which can be the case if we're reading a cache), force which proxy we use.
-								bool bUseParticle = CollisionDataItem.Levelset ? bSwapOrder : true;
-								if (CollisionDataItem.Particle)
+								IPhysicsProxyBase* const PhysicsProxy1 = CollisionDataItem.Proxy2 ? PhysicsProxy0 : CollisionDataItem.Proxy2;
+
+								// Are the proxies pending destruction? If they are no longer tracked by the PhysScene, the proxy is deleted or pending deletion.
+								if (Scene.GetOwningComponent<UPrimitiveComponent>(PhysicsProxy0) == nullptr || Scene.GetOwningComponent<UPrimitiveComponent>(PhysicsProxy1) == nullptr)
 								{
-									IPhysicsProxyBase* const PhysicsProxy1 = bUseParticle ? CollisionDataItem.Particle->PhysicsProxy() : CollisionDataItem.Levelset->PhysicsProxy();
+									continue;
+								}
+
+								{
+									bool bNewEntry = false;
+									FCollisionNotifyInfo& NotifyInfo = GetPendingCollisionForContactPair(PhysicsProxy0, PhysicsProxy1, bNewEntry);
+
+									// #note: we only notify on the first contact, though we will still accumulate the impulse data from subsequent contacts
+									const FVector NormalImpulse = FVector::DotProduct(CollisionDataItem.AccumulatedImpulse, CollisionDataItem.Normal) * CollisionDataItem.Normal;	// project impulse along normal
+									const FVector FrictionImpulse = FVector(CollisionDataItem.AccumulatedImpulse) - NormalImpulse; // friction is component not along contact normal
+									NotifyInfo.RigidCollisionData.TotalNormalImpulse += NormalImpulse;
+									NotifyInfo.RigidCollisionData.TotalFrictionImpulse += FrictionImpulse;
 
+									if (bNewEntry)
 									{
-										bool bNewEntry = false;
-										FCollisionNotifyInfo& NotifyInfo = GetPendingCollisionForContactPair(PhysicsProxy0, PhysicsProxy1, bNewEntry);
+										UPrimitiveComponent* const Comp1 = Scene.GetOwningComponent<UPrimitiveComponent>(PhysicsProxy1);
 
-										// #note: we only notify on the first contact, though we will still accumulate the impulse data from subsequent contacts
-										const FVector NormalImpulse = FVector::DotProduct(CollisionDataItem.AccumulatedImpulse, CollisionDataItem.Normal) * CollisionDataItem.Normal;	// project impulse along normal
-										const FVector FrictionImpulse = FVector(CollisionDataItem.AccumulatedImpulse) - NormalImpulse; // friction is component not along contact normal
-										NotifyInfo.RigidCollisionData.TotalNormalImpulse += NormalImpulse;
-										NotifyInfo.RigidCollisionData.TotalFrictionImpulse += FrictionImpulse;
+										// fill in legacy contact data
+										NotifyInfo.bCallEvent0 = true;
+										// if Comp1 wants this event too, it will get its own pending collision entry, so we leave it false
 
-										if (bNewEntry)
-										{
-											UPrimitiveComponent* const Comp1 = Scene.GetOwningComponent<UPrimitiveComponent>(PhysicsProxy1);
+										SetCollisionInfoFromComp(NotifyInfo.Info0, Comp0);
+										SetCollisionInfoFromComp(NotifyInfo.Info1, Comp1);
 
-											// fill in legacy contact data
-											NotifyInfo.bCallEvent0 = true;
-											// if Comp1 wants this event too, it will get its own pending collision entry, so we leave it false
+										FRigidBodyContactInfo& NewContact = NotifyInfo.RigidCollisionData.ContactInfos.AddZeroed_GetRef();
+										NewContact.ContactNormal = CollisionDataItem.Normal;
+										NewContact.ContactPosition = CollisionDataItem.Location;
+										NewContact.ContactPenetration = CollisionDataItem.PenetrationDepth;
+										// NewContact.PhysMaterial[1] UPhysicalMaterial required here
+									}
 
-											SetCollisionInfoFromComp(NotifyInfo.Info0, Comp0);
-											SetCollisionInfoFromComp(NotifyInfo.Info1, Comp1);
+								}
 
-											FRigidBodyContactInfo& NewContact = NotifyInfo.RigidCollisionData.ContactInfos.AddZeroed_GetRef();
-											NewContact.ContactNormal = CollisionDataItem.Normal;
-											NewContact.ContactPosition = CollisionDataItem.Location;
-											NewContact.ContactPenetration = CollisionDataItem.PenetrationDepth;
-											// NewContact.PhysMaterial[1] UPhysicalMaterial required here
-										}
 
-									}
+								if (HandlerSet.ChaosHandlers.Num() > 0)
+								{
+									bool bNewEntry = false;
+									FChaosPendingCollisionNotify& ChaosNotifyInfo = GetPendingChaosCollisionForContactPair(PhysicsProxy0, PhysicsProxy1, bNewEntry);
 
+									// #note: we only notify on the first contact, though we will still accumulate the impulse data from subsequent contacts
+									ChaosNotifyInfo.CollisionInfo.AccumulatedImpulse += CollisionDataItem.AccumulatedImpulse;
 
-									if (HandlerSet.ChaosHandlers.Num() > 0)
+									if (bNewEntry)
 									{
-										bool bNewEntry = false;
-										FChaosPendingCollisionNotify& ChaosNotifyInfo = GetPendingChaosCollisionForContactPair(PhysicsProxy0, PhysicsProxy1, bNewEntry);
+										UPrimitiveComponent* const Comp1 = Scene.GetOwningComponent<UPrimitiveComponent>(PhysicsProxy1);
 
-										// #note: we only notify on the first contact, though we will still accumulate the impulse data from subsequent contacts
-										ChaosNotifyInfo.CollisionInfo.AccumulatedImpulse += CollisionDataItem.AccumulatedImpulse;
+										// fill in Chaos contact data
+										ChaosNotifyInfo.CollisionInfo.Component = Comp0;
+										ChaosNotifyInfo.CollisionInfo.OtherComponent = Comp1;
+										ChaosNotifyInfo.CollisionInfo.Location = CollisionDataItem.Location;
+										ChaosNotifyInfo.NotifyRecipients = HandlerSet.ChaosHandlers;
 
-										if (bNewEntry)
+										if (bSwapOrder)
 										{
-											UPrimitiveComponent* const Comp1 = Scene.GetOwningComponent<UPrimitiveComponent>(PhysicsProxy1);
-
-											// fill in Chaos contact data
-											ChaosNotifyInfo.CollisionInfo.Component = Comp0;
-											ChaosNotifyInfo.CollisionInfo.OtherComponent = Comp1;
-											ChaosNotifyInfo.CollisionInfo.Location = CollisionDataItem.Location;
-											ChaosNotifyInfo.NotifyRecipients = HandlerSet.ChaosHandlers;
-
-											if (bSwapOrder)
-											{
-												ChaosNotifyInfo.CollisionInfo.AccumulatedImpulse = -CollisionDataItem.AccumulatedImpulse;
-												ChaosNotifyInfo.CollisionInfo.Normal = -CollisionDataItem.Normal;
-
-												ChaosNotifyInfo.CollisionInfo.Velocity = CollisionDataItem.Velocity2;
-												ChaosNotifyInfo.CollisionInfo.OtherVelocity = CollisionDataItem.Velocity1;
-												ChaosNotifyInfo.CollisionInfo.AngularVelocity = CollisionDataItem.AngularVelocity2;
-												ChaosNotifyInfo.CollisionInfo.OtherAngularVelocity = CollisionDataItem.AngularVelocity1;
-												ChaosNotifyInfo.CollisionInfo.Mass = CollisionDataItem.Mass2;
-												ChaosNotifyInfo.CollisionInfo.OtherMass = CollisionDataItem.Mass1;
-											}
-											else
-											{
-												ChaosNotifyInfo.CollisionInfo.AccumulatedImpulse = CollisionDataItem.AccumulatedImpulse;
-												ChaosNotifyInfo.CollisionInfo.Normal = CollisionDataItem.Normal;
-
-												ChaosNotifyInfo.CollisionInfo.Velocity = CollisionDataItem.Velocity1;
-												ChaosNotifyInfo.CollisionInfo.OtherVelocity = CollisionDataItem.Velocity2;
-												ChaosNotifyInfo.CollisionInfo.AngularVelocity = CollisionDataItem.AngularVelocity1;
-												ChaosNotifyInfo.CollisionInfo.OtherAngularVelocity = CollisionDataItem.AngularVelocity2;
-												ChaosNotifyInfo.CollisionInfo.Mass = CollisionDataItem.Mass1;
-												ChaosNotifyInfo.CollisionInfo.OtherMass = CollisionDataItem.Mass2;
-											}
+											ChaosNotifyInfo.CollisionInfo.AccumulatedImpulse = -CollisionDataItem.AccumulatedImpulse;
+											ChaosNotifyInfo.CollisionInfo.Normal = -CollisionDataItem.Normal;
+
+											ChaosNotifyInfo.CollisionInfo.Velocity = CollisionDataItem.Velocity2;
+											ChaosNotifyInfo.CollisionInfo.OtherVelocity = CollisionDataItem.Velocity1;
+											ChaosNotifyInfo.CollisionInfo.AngularVelocity = CollisionDataItem.AngularVelocity2;
+											ChaosNotifyInfo.CollisionInfo.OtherAngularVelocity = CollisionDataItem.AngularVelocity1;
+											ChaosNotifyInfo.CollisionInfo.Mass = CollisionDataItem.Mass2;
+											ChaosNotifyInfo.CollisionInfo.OtherMass = CollisionDataItem.Mass1;
+										}
+										else
+										{
+											ChaosNotifyInfo.CollisionInfo.AccumulatedImpulse = CollisionDataItem.AccumulatedImpulse;
+											ChaosNotifyInfo.CollisionInfo.Normal = CollisionDataItem.Normal;
+
+											ChaosNotifyInfo.CollisionInfo.Velocity = CollisionDataItem.Velocity1;
+											ChaosNotifyInfo.CollisionInfo.OtherVelocity = CollisionDataItem.Velocity2;
+											ChaosNotifyInfo.CollisionInfo.AngularVelocity = CollisionDataItem.AngularVelocity1;
+											ChaosNotifyInfo.CollisionInfo.OtherAngularVelocity = CollisionDataItem.AngularVelocity2;
+											ChaosNotifyInfo.CollisionInfo.Mass = CollisionDataItem.Mass1;
+											ChaosNotifyInfo.CollisionInfo.OtherMass = CollisionDataItem.Mass2;
 										}
 									}
 								}
@@ -394,7 +395,6 @@ void UChaosGameplayEventDispatcher::HandleCollisionEvents(const Chaos::FCollisio
 
 void UChaosGameplayEventDispatcher::HandleBreakingEvents(const Chaos::FBreakingEventData& Event)
 {
-
 	SCOPE_CYCLE_COUNTER(STAT_DispatchBreakEvents);
 
 	// BREAK EVENTS
@@ -414,9 +414,9 @@ void UChaosGameplayEventDispatcher::HandleBreakingEvents(const Chaos::FBreakingE
 		{
 			for (Chaos::FBreakingData const& BreakingDataItem : BreakingData)
 			{	
-				if (BreakingDataItem.Particle)
+				if (BreakingDataItem.Proxy)
 				{
-					UPrimitiveComponent* const PrimComp = Cast<UPrimitiveComponent>(BreakingDataItem.Particle->PhysicsProxy()->GetOwner());
+					UPrimitiveComponent* const PrimComp = Cast<UPrimitiveComponent>(BreakingDataItem.Proxy->GetOwner());
 					if (PrimComp && BreakEventRegistrations.Contains(PrimComp))
 					{
 						// queue them up so we can release the physics data before trigging BP events
@@ -436,16 +436,17 @@ void UChaosGameplayEventDispatcher::HandleBreakingEvents(const Chaos::FBreakingE
 
 }
 
-
 void UChaosGameplayEventDispatcher::HandleSleepingEvents(const Chaos::FSleepingEventData& SleepingData)
 {
+	FPhysScene_Chaos& Scene = *(GetWorld()->GetPhysicsScene());
+
 	const Chaos::FSleepingDataArray& SleepingArray = SleepingData.SleepingData;
 
 	for (const Chaos::FSleepingData& SleepData : SleepingArray)
 	{
-		if (SleepData.Particle->GetProxy()!= nullptr)
+		if (SleepData.Proxy!= nullptr && Scene.GetOwningComponent<UPrimitiveComponent>(SleepData.Proxy) != nullptr)
 		{
-			if (FBodyInstance* BodyInstance = FPhysicsUserData::Get<FBodyInstance>(SleepData.Particle->UserData()))
+			if (FBodyInstance* BodyInstance = Scene.GetBodyInstanceFromProxy(SleepData.Proxy))
 			{
 				if (BodyInstance->bGenerateWakeEvents)
 				{
@@ -459,7 +460,6 @@ void UChaosGameplayEventDispatcher::HandleSleepingEvents(const Chaos::FSleepingE
 	DispatchPendingWakeNotifies();
 }
 
-
 void UChaosGameplayEventDispatcher::AddPendingSleepingNotify(FBodyInstance* BodyInstance, ESleepEvent SleepEventType)
 {
 	PendingSleepNotifies.FindOrAdd(BodyInstance) = SleepEventType;
diff --git a/Engine/Source/Runtime/Experimental/GeometryCollectionEngine/Private/ChaosCollisionEventFilter.cpp b/Engine/Source/Runtime/Experimental/GeometryCollectionEngine/Private/ChaosCollisionEventFilter.cpp
--- a/Engine/Source/Runtime/Experimental/GeometryCollectionEngine/Private/ChaosCollisionEventFilter.cpp
+++ b/Engine/Source/Runtime/Experimental/GeometryCollectionEngine/Private/ChaosCollisionEventFilter.cpp
@@ -59,8 +59,6 @@ void FChaosCollisionEventFilter::FilterEvents(const FTransform& ChaosComponentTr
 			NewData.Mass1 = CollisionData.Mass1;
 			NewData.Mass2 = CollisionData.Mass2;
 			NewData.Impulse = CollisionData.AccumulatedImpulse;
-			NewData.Particle = CollisionData.Particle;
-			NewData.Levelset = CollisionData.Levelset;
 
 			FilteredDataArray.Add(NewData);
 
diff --git a/Engine/Source/Runtime/Experimental/GeometryCollectionEngine/Public/ChaosCollisionEventFilter.h b/Engine/Source/Runtime/Experimental/GeometryCollectionEngine/Public/ChaosCollisionEventFilter.h
--- a/Engine/Source/Runtime/Experimental/GeometryCollectionEngine/Public/ChaosCollisionEventFilter.h
+++ b/Engine/Source/Runtime/Experimental/GeometryCollectionEngine/Public/ChaosCollisionEventFilter.h
@@ -41,14 +41,6 @@ struct FChaosCollisionEventData
 	UPROPERTY(EditAnywhere, BlueprintReadOnly, Category = "Data")
 	FVector Impulse;
 
-	// The particle of the collision event
-	//UPROPERTY(EditAnywhere, BlueprintReadOnly, Category = "Data")
-	Chaos::FGeometryParticleHandle* Particle;
-
-	// The level set of the collision event
-	//UPROPERTY(EditAnywhere, BlueprintReadOnly, Category = "Data")
-	Chaos::FGeometryParticleHandle* Levelset;
-
 	FChaosCollisionEventData()
 		: Location(FVector::ZeroVector)
 		, Normal(FVector::ZeroVector)
@@ -57,8 +49,6 @@ struct FChaosCollisionEventData
 		, Mass1(0.0f)
 		, Mass2(0.0f)
 		, Impulse(FVector::ZeroVector)
-		, Particle(nullptr)
-		, Levelset(nullptr)
 	{
 	}
 };
