Chaos Caching
-- added asset registry tags to CacheCollection asset

diff --git a/Engine/Plugins/Experimental/ChaosCaching/Source/ChaosCaching/Private/Chaos/CacheCollection.cpp b/Engine/Plugins/Experimental/ChaosCaching/Source/ChaosCaching/Private/Chaos/CacheCollection.cpp
--- a/Engine/Plugins/Experimental/ChaosCaching/Source/ChaosCaching/Private/Chaos/CacheCollection.cpp
+++ b/Engine/Plugins/Experimental/ChaosCaching/Source/ChaosCaching/Private/Chaos/CacheCollection.cpp
@@ -5,6 +5,25 @@
 #include "Algo/Find.h"
 #include "Async/ParallelFor.h"
 
+void UChaosCacheCollection::GetAssetRegistryTags(TArray<FAssetRegistryTag>& OutTags) const
+{
+#if WITH_EDITOR
+	int32 NumCaches = Caches.Num();
+
+	OutTags.Add(FAssetRegistryTag(TEXT("Number of Observed Components"), FString::FromInt(NumCaches), FAssetRegistryTag::TT_Numerical));
+
+	float Duration = 0.0;
+	for (int32 CacheIdx = 0; CacheIdx < NumCaches; ++CacheIdx)
+	{
+		UChaosCache* Cache = Caches[CacheIdx];
+		Duration = FMath::Max(Duration, Cache->RecordedDuration);
+	}	
+	OutTags.Add(FAssetRegistryTag(TEXT("Recorded Duration"), FString::Printf(TEXT("%.2f"), Duration), FAssetRegistryTag::TT_Numerical));
+#endif
+
+	Super::GetAssetRegistryTags(OutTags);
+}
+
 UChaosCache* UChaosCacheCollection::FindCache(const FName& CacheName) const
 {
 	UChaosCache* const* ExistingCache = Algo::FindByPredicate(Caches, [&CacheName](const UChaosCache* Test)
diff --git a/Engine/Plugins/Experimental/ChaosCaching/Source/ChaosCaching/Public/Chaos/CacheCollection.h b/Engine/Plugins/Experimental/ChaosCaching/Source/ChaosCaching/Public/Chaos/CacheCollection.h
--- a/Engine/Plugins/Experimental/ChaosCaching/Source/ChaosCaching/Public/Chaos/CacheCollection.h
+++ b/Engine/Plugins/Experimental/ChaosCaching/Source/ChaosCaching/Public/Chaos/CacheCollection.h
@@ -12,6 +12,10 @@ class CHAOSCACHING_API UChaosCacheCollection : public UObject
 	GENERATED_BODY()
 public:
 
+	//~ Begin UObject Interface
+	virtual void GetAssetRegistryTags(TArray<FAssetRegistryTag>& OutTags) const override;
+	//~ End UObject Interface
+	
 	UChaosCache* FindCache(const FName& CacheName) const;
 	UChaosCache* FindOrAddCache(const FName& CacheName);
 
