Chaos Collection
-- Create CacheCollection asset when the CacheManagerActor is instantiated.

diff --git a/Engine/Plugins/Experimental/ChaosCaching/Source/ChaosCachingEditor/Private/Chaos/ChaosCachingEditorPlugin.cpp b/Engine/Plugins/Experimental/ChaosCaching/Source/ChaosCachingEditor/Private/Chaos/ChaosCachingEditorPlugin.cpp
--- a/Engine/Plugins/Experimental/ChaosCaching/Source/ChaosCachingEditor/Private/Chaos/ChaosCachingEditorPlugin.cpp
+++ b/Engine/Plugins/Experimental/ChaosCaching/Source/ChaosCachingEditor/Private/Chaos/ChaosCachingEditorPlugin.cpp
@@ -9,6 +9,8 @@
 #include "Chaos/CacheEditorCommands.h"
 #include "Chaos/CacheManagerActor.h"
 #include "Chaos/CacheManagerCustomization.h"
+#include "Chaos/CacheCollectionFactory.h"
+#include "Chaos/CacheCollection.h"
 #include "CoreMinimal.h"
 #include "Engine/Selection.h"
 #include "Kismet2/ComponentEditorUtils.h"
@@ -181,6 +183,20 @@ void IChaosCachingEditorPlugin::OnCreateCacheManager()
 			}
 		}
 	}
+
+	// Create an associated Cache Collection
+	if (Manager->ObservedComponents.Num() > 0)
+	{
+		FAssetToolsModule& AssetToolsModule = FModuleManager::GetModuleChecked<FAssetToolsModule>("AssetTools");
+		UCacheCollectionFactory* Factory = NewObject<UCacheCollectionFactory>();
+		UChaosCacheCollection* NewAsset = Cast<UChaosCacheCollection>(AssetToolsModule.Get().CreateAssetWithDialog(UChaosCacheCollection::StaticClass(), Factory));
+
+		if (NewAsset)
+		{
+			Manager->CacheCollection = NewAsset;
+		}
+	}
+
 }
 
 bool IsCreateCacheManagerVisible()
