Chaos Caching
-- correctly gate Manager pointer to protect against deref of nullptr

diff --git a/Engine/Plugins/Experimental/ChaosCaching/Source/ChaosCachingEditor/Private/Chaos/ChaosCachingEditorPlugin.cpp b/Engine/Plugins/Experimental/ChaosCaching/Source/ChaosCachingEditor/Private/Chaos/ChaosCachingEditorPlugin.cpp
--- a/Engine/Plugins/Experimental/ChaosCaching/Source/ChaosCachingEditor/Private/Chaos/ChaosCachingEditorPlugin.cpp
+++ b/Engine/Plugins/Experimental/ChaosCaching/Source/ChaosCachingEditor/Private/Chaos/ChaosCachingEditorPlugin.cpp
@@ -184,19 +184,21 @@ void IChaosCachingEditorPlugin::OnCreateCacheManager()
 		}
 	}
 
-	// Create an associated Cache Collection
-	if (Manager->ObservedComponents.Num() > 0)
+	if (Manager)
 	{
-		FAssetToolsModule& AssetToolsModule = FModuleManager::GetModuleChecked<FAssetToolsModule>("AssetTools");
-		UCacheCollectionFactory* Factory = NewObject<UCacheCollectionFactory>();
-		UChaosCacheCollection* NewAsset = Cast<UChaosCacheCollection>(AssetToolsModule.Get().CreateAssetWithDialog(UChaosCacheCollection::StaticClass(), Factory));
-
-		if (NewAsset)
+		// Create an associated Cache Collection
+		if (Manager->ObservedComponents.Num() > 0)
 		{
-			Manager->CacheCollection = NewAsset;
+			FAssetToolsModule& AssetToolsModule = FModuleManager::GetModuleChecked<FAssetToolsModule>("AssetTools");
+			UCacheCollectionFactory* Factory = NewObject<UCacheCollectionFactory>();
+			UChaosCacheCollection* NewAsset = Cast<UChaosCacheCollection>(AssetToolsModule.Get().CreateAssetWithDialog(UChaosCacheCollection::StaticClass(), Factory));
+
+			if (NewAsset)
+			{
+				Manager->CacheCollection = NewAsset;
+			}
 		}
 	}
-
 }
 
 bool IsCreateCacheManagerVisible()
