Chaos Caching
-- changed ensures to be accurate to Game/Physics thread context

diff --git a/Engine/Plugins/Experimental/ChaosCaching/Source/ChaosCaching/Private/Chaos/Adapters/GeometryCollectionComponentCacheAdapter.cpp b/Engine/Plugins/Experimental/ChaosCaching/Source/ChaosCaching/Private/Chaos/Adapters/GeometryCollectionComponentCacheAdapter.cpp
--- a/Engine/Plugins/Experimental/ChaosCaching/Source/ChaosCaching/Private/Chaos/Adapters/GeometryCollectionComponentCacheAdapter.cpp
+++ b/Engine/Plugins/Experimental/ChaosCaching/Source/ChaosCaching/Private/Chaos/Adapters/GeometryCollectionComponentCacheAdapter.cpp
@@ -44,7 +44,7 @@ namespace Chaos
 
 	void FGeometryCollectionCacheAdapter::Record_PostSolve(UPrimitiveComponent* InComp, const FTransform& InRootTransform, FPendingFrameWrite& OutFrame, Chaos::FReal InTime) const
 	{
-		ensure(IsInGameThreadContext());
+		ensure(IsInPhysicsThreadContext());
 
 		using FClusterParticle = Chaos::FPBDRigidClusteredParticleHandle;
 		using FRigidParticle = Chaos::FPBDRigidParticleHandle;
@@ -191,7 +191,7 @@ namespace Chaos
 															FPlaybackTickRecord&                               TickRecord,
 															TArray<TPBDRigidParticleHandle<Chaos::FReal, 3>*>& OutUpdatedRigids) const
 	{
-		ensure(IsInGameThreadContext());
+		ensure(IsInPhysicsThreadContext());
 		
 		using FClusterParticle = Chaos::FPBDRigidClusteredParticleHandle;
 		using FRigidParticle = Chaos::FPBDRigidParticleHandle;
