ChaosCaching
-- Removed per observed controls for StartMode, CacheMode and moved to a single set of controls on the cache manager.
-- Cache takes name from observed component.
-- Placed CacheCollection automatically set to Kinematic, non-placed are Record
-- Random access start frame
-- Consistenty record all transforms in Actor space
-- Compatible with Sequencer control
-- Created Playback only CacheManager subclass
-- GeometryCollection recording abides component flags to record/playback breaking/collision/trailing events

diff --git a/Engine/Plugins/Experimental/ChaosCaching/Source/ChaosCaching/Private/Chaos/Adapters/GeometryCollectionComponentCacheAdapter.cpp b/Engine/Plugins/Experimental/ChaosCaching/Source/ChaosCaching/Private/Chaos/Adapters/GeometryCollectionComponentCacheAdapter.cpp
--- a/Engine/Plugins/Experimental/ChaosCaching/Source/ChaosCaching/Private/Chaos/Adapters/GeometryCollectionComponentCacheAdapter.cpp
+++ b/Engine/Plugins/Experimental/ChaosCaching/Source/ChaosCaching/Private/Chaos/Adapters/GeometryCollectionComponentCacheAdapter.cpp
@@ -8,6 +8,7 @@
 #include "PhysicsProxy/GeometryCollectionPhysicsProxy.h"
 #include "PhysicsProxy/SingleParticlePhysicsProxy.h"
 #include "Chaos/ChaosSolverActor.h"
+#include "GeometryCollection/GeometryCollectionAlgo.h"
 
 FName FEnableStateEvent::EventName("GC_Enable");
 FName FBreakingEvent::EventName("GC_Breaking");
@@ -41,7 +42,7 @@ namespace Chaos
 	{
 		return EngineAdapterPriotityBegin;
 	}
-
+	
 	void FGeometryCollectionCacheAdapter::Record_PostSolve(UPrimitiveComponent* InComp, const FTransform& InRootTransform, FPendingFrameWrite& OutFrame, Chaos::FReal InTime) const
 	{
 		ensure(IsInPhysicsThreadContext());
@@ -76,9 +77,12 @@ namespace Chaos
 		}
 
 		FGeometryDynamicCollection&            Collection       = Proxy->GetPhysicsCollection();
+		const TManagedArray<FTransform>&       MassToLocal		= Collection.MassToLocal;
+		const TManagedArray<int32>&            Parents			= RestCollection->Parent;
 		const TManagedArray<int32>&            TransformIndices = RestCollection->TransformIndex;
 		const TManagedArray<FTransform>&       Transforms       = Collection.Transform;
-		const TArray<FBreakingData>& Breaks                     = Solver->GetEvolution()->GetRigidClustering().GetAllClusterBreakings();
+		const TArray<FBreakingData>&		   Breaks           = Solver->GetEvolution()->GetRigidClustering().GetAllClusterBreakings();
+		const FTransform					   WorldToActor		= InRootTransform.Inverse();
 
 		// A transform index exists for each 'real' (i.e. leaf node in the rest collection)
 		const int32 NumTransforms = Collection.NumElements(FGeometryCollection::TransformGroup);
@@ -86,8 +90,14 @@ namespace Chaos
 		// Pre-alloc once for worst case.
 		OutFrame.PendingParticleData.Reserve(NumTransforms);
 
+		// Record the cluster indices that released this frame. Those particles become
+		// inactive: to make random access reads convenient, we set their transforms to identity.
+		TSet<int32> ReleasedClusters;
+		ReleasedClusters.Reserve(Breaks.Num());
+		
 		TArray<int32> RelatedBreaks;
 		RelatedBreaks.Reserve(Breaks.Num());
+		
 		for(const FBreakingData& Break : Breaks)
 		{
 			// Is the proxy pending destruction? If they are no longer tracked by the PhysScene, the proxy is deleted or pending deletion.
@@ -104,6 +114,9 @@ namespace Chaos
 				{
 					// The break particle belongs to our proxy
 					RelatedBreaks.Add(Break.TransformGroupIndex);
+
+					// Record the cluster it belonged to.
+					ReleasedClusters.Add(Parents[Break.TransformGroupIndex]);
 				}
 			}
 		}
@@ -126,7 +139,11 @@ namespace Chaos
 					FPendingParticleWrite& Pending = OutFrame.PendingParticleData.Last();
 
 					Pending.ParticleIndex = TransformIndex;
-					Pending.PendingTransform = FTransform(Handle->R(), Handle->X()).GetRelativeTransform(InRootTransform);
+					
+					// All recorded transforms are in actor space, ie relative to the Cache Manager making the recording.
+					FTransform LocalTransform = MassToLocal[TransformIndex].Inverse() * FTransform(Handle->R(), Handle->X());
+					FTransform ActorSpaceTransform = LocalTransform * WorldToActor;
+					Pending.PendingTransform = ActorSpaceTransform;
 				}
 
 				if (RelatedBreaks.Contains(TransformIndex))
@@ -136,7 +153,26 @@ namespace Chaos
 			}
 		}
 
-		if (BreakingDataArray && ProxyCachedEventData.ProxyBreakingDataIndices)
+		// If a cluster particle released, set its transform to identity.
+		for (int32 TransformIndex : ReleasedClusters)
+		{
+			if (TransformIndex > INDEX_NONE)
+			{
+				FClusterParticle* Handle = Proxy->GetParticles()[TransformIndex];
+				if (Handle)
+				{
+					OutFrame.PendingParticleData.AddDefaulted();
+					FPendingParticleWrite& Pending = OutFrame.PendingParticleData.Last();
+					Pending.ParticleIndex = TransformIndex;
+					FTransform LocalTransform = MassToLocal[TransformIndex].Inverse() * FTransform(Handle->R(), Handle->X());
+					FTransform ActorSpaceTransform = LocalTransform * WorldToActor;
+					Pending.PendingTransform = ActorSpaceTransform;
+					Pending.bPendingDeactivate = true;
+				}
+			}
+		}
+
+		if (bHandleBreaks && BreakingDataArray && ProxyCachedEventData.ProxyBreakingDataIndices)
 		{
 			for (int32 Index : *ProxyCachedEventData.ProxyBreakingDataIndices)
 			{
@@ -152,7 +188,7 @@ namespace Chaos
 			}
 		}
 		
-		if (CollisionDataArray && ProxyCachedEventData.ProxyCollisionDataIndices)
+		if (bHandleCollisions && CollisionDataArray && ProxyCachedEventData.ProxyCollisionDataIndices)
 		{
 			for (int32 Index : *ProxyCachedEventData.ProxyCollisionDataIndices)
 			{
@@ -165,7 +201,7 @@ namespace Chaos
 			}
 		}
 		
-		if (TrailingDataArray && ProxyCachedEventData.ProxyTrailingDataIndices)
+		if (bHandleTrailing && TrailingDataArray && ProxyCachedEventData.ProxyTrailingDataIndices)
 		{
 			for (int32 Index : *ProxyCachedEventData.ProxyTrailingDataIndices)
 			{
@@ -184,11 +220,11 @@ namespace Chaos
 		// Never going to change again till freed after writing to the cache so free up the extra space we reserved
 		OutFrame.PendingParticleData.Shrink();
 	}
-
-	void FGeometryCollectionCacheAdapter::Playback_PreSolve(UPrimitiveComponent*                               InComponent,
-															UChaosCache*                                       InCache,
+	
+	void FGeometryCollectionCacheAdapter::Playback_PreSolve(UPrimitiveComponent*							   InComponent,
+															UChaosCache*									   InCache,
 															Chaos::FReal                                       InTime,
-															FPlaybackTickRecord&                               TickRecord,
+															FPlaybackTickRecord&							   TickRecord,
 															TArray<TPBDRigidParticleHandle<Chaos::FReal, 3>*>& OutUpdatedRigids) const
 	{
 		ensure(IsInPhysicsThreadContext());
@@ -214,6 +250,7 @@ namespace Chaos
 		}
 
 		FGeometryDynamicCollection&      Collection       = Proxy->GetPhysicsCollection();
+		const TManagedArray<FTransform>& MassToLocal	  = Collection.MassToLocal;
 		const TManagedArray<int32>&      TransformIndices = RestCollection->TransformIndex;
 		const TManagedArray<FTransform>& Transforms       = Collection.Transform;
 		TArray<FClusterParticle*>        Particles        = Proxy->GetParticles();
@@ -300,7 +337,7 @@ namespace Chaos
 			}
 		}
 
-		if (BreakingEvents)
+		if (bHandleBreaks && BreakingEvents)
 		{
 			const FSolverBreakingEventFilter* SolverBreakingEventFilter = Solver->GetEventFilters()->GetBreakingFilter();
 
@@ -348,7 +385,7 @@ namespace Chaos
 			}
 		}
 
-		if (TrailingEvents)
+		if (bHandleTrailing && TrailingEvents)
 		{
 			const FSolverTrailingEventFilter* SolverTrailingEventFilter = Solver->GetEventFilters()->GetTrailingFilter();
 
@@ -395,7 +432,7 @@ namespace Chaos
 			}
 		}
 
-		if (CollisionEvents)
+		if (bHandleCollisions && CollisionEvents)
 		{
 
 			const FSolverCollisionEventFilter* SolverCollisionEventFilter = Solver->GetEventFilters()->GetCollisionFilter();
@@ -455,8 +492,12 @@ namespace Chaos
 					continue;
 				}
 
-				Handle->SetP(EvaluatedTransform.GetTranslation());
-				Handle->SetQ(EvaluatedTransform.GetRotation());
+				// Evaluated transform is in Cache Manager space with MassToLocal removed. We need it in World space.
+				// Note that before Version 1, MassToLocal is burned into the evaluated transform.
+				FTransform WorldTransform = (InCache->Version < 1) ? EvaluatedTransform : MassToLocal[ParticleIndex] * EvaluatedTransform;
+				
+				Handle->SetP(WorldTransform.GetTranslation());
+				Handle->SetQ(WorldTransform.GetRotation());
 				Handle->SetX(Handle->P());
 				Handle->SetR(Handle->Q());
 				
@@ -558,11 +599,77 @@ namespace Chaos
 	{
 		CachedData.Empty();
 	}
-	
+
+	void FGeometryCollectionCacheAdapter::SetRestState(UPrimitiveComponent* InComponent, UChaosCache* InCache, const FTransform& InRootTransform, Chaos::FReal InTime) const
+	{
+		// Caches recorded previous to Version 1 may not scrub correctly as the MassToLocal transform has been burned in.
+		
+		ensure(IsInGameThreadContext());
+		
+		if (!InCache || InCache->GetDuration() == 0.0f)
+		{
+			return;
+		}
+
+		if (UGeometryCollectionComponent* Comp = CastChecked<UGeometryCollectionComponent>(InComponent))
+		{
+			const FTransform ActorToWorld = InRootTransform;
+			const FTransform WorldToComponent = Comp->GetComponentTransform().Inverse();
+			FTransform ActorToComponent = ActorToWorld * WorldToComponent;
+
+			
+			if (const UGeometryCollection* RestCollection = Comp->GetRestCollection())
+			{
+				TSharedPtr<FGeometryCollection, ESPMode::ThreadSafe> GeometryCollection = RestCollection->GetGeometryCollection();
+				const int32 NumTransforms = GeometryCollection->NumElements(FGeometryCollection::TransformGroup);
+
+				FPlaybackTickRecord TickRecord;
+				TickRecord.SetLastTime(InTime);
+				FCacheEvaluationContext Context(TickRecord);
+				Context.bEvaluateTransform = true;
+				Context.bEvaluateCurves = false;
+				Context.bEvaluateEvents = false;
+
+				FCacheEvaluationResult EvaluatedResult = InCache->Evaluate(Context);
+				const int32 NumCacheTransforms = EvaluatedResult.Transform.Num();
+			
+				// Any bone that is not explicitly set by the cache defaults to it's rest collection position.
+				TArray<FTransform> RestTransforms;
+				RestTransforms.SetNum(NumTransforms);
+				for (int32 Idx = 0; Idx < NumTransforms; ++Idx)
+				{
+					RestTransforms[Idx] = GeometryCollection->Transform[Idx];
+				}
+
+				for (int32 CacheIdx = 0; CacheIdx < NumCacheTransforms; ++CacheIdx)
+				{
+					const int32 TransformIdx = EvaluatedResult.ParticleIndices[CacheIdx];
+					RestTransforms[TransformIdx] = EvaluatedResult.Transform[CacheIdx] * ActorToComponent;
+				}
+
+				// Set any broken clusters to identity so that they no longer influence child transforms.
+				TArray<int32> ReleaseIndices = GatherAllBreaksUpToTime(InCache, InTime);
+				const TManagedArray<int32>& Parent = GeometryCollection->Parent;
+				for (int32 ReleaseIdx : ReleaseIndices)
+				{
+					int32 ClusterIdx = Parent[ReleaseIdx];
+					if (ClusterIdx > INDEX_NONE)
+					{
+						RestTransforms[ClusterIdx] = FTransform::Identity;
+					}
+				}
+
+				Comp->SetRestState(MoveTemp(RestTransforms));
+			}
+		}
+	}
+
 	bool FGeometryCollectionCacheAdapter::InitializeForRecord(UPrimitiveComponent* InComponent, UChaosCache* InCache)
 	{
 		ensure(IsInGameThreadContext());
 		
+		InCache->Version = 1;
+		
 		UGeometryCollectionComponent*    Comp     = CastChecked<UGeometryCollectionComponent>(InComponent);
 		FGeometryCollectionPhysicsProxy* Proxy    = Comp->GetPhysicsProxy();
 
@@ -578,6 +685,10 @@ namespace Chaos
 			return false;
 		}
 
+		bHandleBreaks = Comp->bNotifyBreaks;
+		bHandleCollisions = Comp->bNotifyCollisions;
+		bHandleTrailing = Comp->bNotifyTrailing;
+		
 		// We need secondary event data to record event information into the cache
 		Solver->SetGenerateBreakingData(true);
 		Solver->SetGenerateCollisionData(true);
@@ -589,9 +700,20 @@ namespace Chaos
 			Chaos::FEventManager* EventManager = Solver->GetEventManager();
 			if (EventManager)
 			{
-				EventManager->RegisterHandler<Chaos::FBreakingEventData>(Chaos::EEventType::Breaking, const_cast<FGeometryCollectionCacheAdapter*>(this), &FGeometryCollectionCacheAdapter::HandleBreakingEvents);
-				EventManager->RegisterHandler<Chaos::FCollisionEventData>(Chaos::EEventType::Collision, const_cast<FGeometryCollectionCacheAdapter*>(this), &FGeometryCollectionCacheAdapter::HandleCollisionEvents);
-				EventManager->RegisterHandler<Chaos::FTrailingEventData>(Chaos::EEventType::Trailing, const_cast<FGeometryCollectionCacheAdapter*>(this), &FGeometryCollectionCacheAdapter::HandleTrailingEvents);
+				if (bHandleBreaks)
+				{ 
+					EventManager->RegisterHandler<Chaos::FBreakingEventData>(Chaos::EEventType::Breaking, const_cast<FGeometryCollectionCacheAdapter*>(this), &FGeometryCollectionCacheAdapter::HandleBreakingEvents);
+				}
+
+				if (bHandleCollisions)
+				{ 
+					EventManager->RegisterHandler<Chaos::FCollisionEventData>(Chaos::EEventType::Collision, const_cast<FGeometryCollectionCacheAdapter*>(this), &FGeometryCollectionCacheAdapter::HandleCollisionEvents);
+				}
+
+				if (bHandleTrailing)
+				{ 
+					EventManager->RegisterHandler<Chaos::FTrailingEventData>(Chaos::EEventType::Trailing, const_cast<FGeometryCollectionCacheAdapter*>(this), &FGeometryCollectionCacheAdapter::HandleTrailingEvents);
+				}
 			}
 
 			BreakingDataArray = nullptr;
@@ -604,30 +726,96 @@ namespace Chaos
 		return true;
 	}
 
-	bool FGeometryCollectionCacheAdapter::InitializeForPlayback(UPrimitiveComponent* InComponent, UChaosCache* InCache) const
+	bool FGeometryCollectionCacheAdapter::InitializeForPlayback(UPrimitiveComponent* InComponent, UChaosCache* InCache, float InTime)
 	{
 		ensure(IsInGameThreadContext());
-		
-		UGeometryCollectionComponent*    Comp  = CastChecked<UGeometryCollectionComponent>(InComponent);
-		if (!Comp)
+
+		if (UGeometryCollectionComponent* Comp = CastChecked<UGeometryCollectionComponent>(InComponent))
 		{
-			return false;
+			bHandleBreaks = Comp->bNotifyBreaks;
+			bHandleCollisions = Comp->bNotifyCollisions;
+			bHandleTrailing = Comp->bNotifyTrailing;
+			
+			if (const UGeometryCollection* RestCollection = Comp->GetRestCollection())
+			{
+				// Set up initial conditions... 
+				
+				TArray<FTransform> InitialTransforms;
+				{
+					FPlaybackTickRecord TickRecord;
+					TickRecord.SetLastTime(InTime);
+					FCacheEvaluationContext Context(TickRecord);
+					Context.bEvaluateTransform = true;
+					Context.bEvaluateCurves = false;
+					Context.bEvaluateEvents = false;
+
+					FCacheEvaluationResult EvaluatedResult = InCache->Evaluate(Context);
+					const int32 NumCacheTransforms = EvaluatedResult.Transform.Num();
+
+					// Any bone that is not explicitly set by the cache defaults to it's rest collection position.
+					TSharedPtr<FGeometryCollection, ESPMode::ThreadSafe> GeometryCollection = RestCollection->GetGeometryCollection();
+					const int32 NumTransforms = GeometryCollection->NumElements(FGeometryCollection::TransformGroup);
+
+					InitialTransforms.SetNum(NumTransforms);
+					for (int32 Idx = 0; Idx < NumTransforms; ++Idx)
+					{
+						InitialTransforms[Idx] = GeometryCollection->Transform[Idx];
+					}
+
+					for (int32 CacheIdx = 0; CacheIdx < NumCacheTransforms; ++CacheIdx)
+					{
+						const int32 TransformIdx = EvaluatedResult.ParticleIndices[CacheIdx];
+						InitialTransforms[TransformIdx] = EvaluatedResult.Transform[CacheIdx];
+					}
+				}
+				Comp->SetInitialTransforms(InitialTransforms);
+					
+				TArray<int32> ReleaseIndices = GatherAllBreaksUpToTime(InCache, InTime);
+				Comp->SetInitialClusterBreaks(ReleaseIndices);
+
+				Comp->SetDynamicState(Chaos::EObjectStateType::Kinematic);
+		
+				// ...then initialize the proxy and particles.
+				Comp->SetSimulatePhysics(true);
+
+				return true;
+			}
 		}
 
-		FGeometryCollectionPhysicsProxy* Proxy = Comp->GetPhysicsProxy();
+		return false;
+	}
 
-		if (ensure(Proxy))
-		{
-			FGeometryDynamicCollection& Collection = Proxy->GetPhysicsCollection();
+	TArray<int32> FGeometryCollectionCacheAdapter::GatherAllBreaksUpToTime(UChaosCache* InCache, float InTime) const
+	{
+		// Evaluate all breaking event that have occured from the beginning of the cache up to the specified time.
+		TArray<int32> ReleaseIndices;
+
+		FPlaybackTickRecord TickRecord;
+		TickRecord.SetLastTime(0.0);
+		TickRecord.SetDt(InTime);
+		FCacheEvaluationContext Context(TickRecord);
+		Context.bEvaluateTransform = false;
+		Context.bEvaluateCurves = false;
+		Context.bEvaluateEvents = true;
 
-			const int32 NumDynamicState = Collection.DynamicState.Num();
-			for (int32 StateIdx = 0; StateIdx < NumDynamicState; ++StateIdx)
+		FCacheEvaluationResult EvaluatedResult = InCache->Evaluate(Context);
+		const TArray<FCacheEventHandle>* EnableEvents = EvaluatedResult.Events.Find(FEnableStateEvent::EventName);
+		if (EnableEvents)
+		{
+			ReleaseIndices.Reserve(EnableEvents->Num());
+			for (const FCacheEventHandle& Handle : *EnableEvents)
 			{
-				Collection.DynamicState[StateIdx] = (int32)EObjectStateTypeEnum::Chaos_Object_Kinematic;
+				if (FEnableStateEvent* Event = Handle.Get<FEnableStateEvent>())
+				{
+					if (Event->bEnable)
+					{
+						ReleaseIndices.Add(Event->Index);
+					}
+				}
 			}
 		}
-		
-		return true;		
+
+		return ReleaseIndices;
 	}
 
 	void FGeometryCollectionCacheAdapter::HandleBreakingEvents(const Chaos::FBreakingEventData& Event)
diff --git a/Engine/Plugins/Experimental/ChaosCaching/Source/ChaosCaching/Private/Chaos/Adapters/StaticMeshComponentCacheAdapter.cpp b/Engine/Plugins/Experimental/ChaosCaching/Source/ChaosCaching/Private/Chaos/Adapters/StaticMeshComponentCacheAdapter.cpp
--- a/Engine/Plugins/Experimental/ChaosCaching/Source/ChaosCaching/Private/Chaos/Adapters/StaticMeshComponentCacheAdapter.cpp
+++ b/Engine/Plugins/Experimental/ChaosCaching/Source/ChaosCaching/Private/Chaos/Adapters/StaticMeshComponentCacheAdapter.cpp
@@ -64,7 +64,7 @@ namespace Chaos
 #endif // WITH_CHAOS
 	}
 
-	void PlayFromCacheInternal(FSingleParticlePhysicsProxy* InProxy, UChaosCache* InCache, FPlaybackTickRecord& TickRecord, TArray<TPBDRigidParticleHandle<Chaos::FReal, 3>*>& OutUpdatedRigids)
+	void PlayFromCacheInternal(FSingleParticlePhysicsProxy* InProxy, UChaosCache* InCache, FPlaybackTickRecord& TickRecord, TArray<TPBDRigidParticleHandle<Chaos::FReal, 3>*>& OutUpdatedRigids) 
 	{
 		if(!InCache || InCache->GetDuration() == 0.0f)
 		{
@@ -140,12 +140,38 @@ namespace Chaos
 		return nullptr;
 	}
 
+	void FStaticMeshCacheAdapter::SetRestState(UPrimitiveComponent* InComponent, UChaosCache* InCache, const FTransform& InRootTransform, Chaos::FReal InTime) const
+	{
+		if (!InCache || InCache->GetDuration() == 0.0f)
+		{
+			return;
+		}
+
+		if (InCache->ParticleTracks.Num() == 1)
+		{
+
+			FPlaybackTickRecord TickRecord;
+			TickRecord.SetLastTime(InTime);
+			FCacheEvaluationContext Context(TickRecord);
+			Context.bEvaluateTransform = true;
+			Context.bEvaluateCurves = false;
+			Context.bEvaluateEvents = false;
+
+			FTransform RestTransform;
+			InCache->EvaluateSingle(0, TickRecord, &RestTransform, nullptr);
+
+			// Evaluated transform is in CacheManager space.
+			InComponent->SetWorldTransform(RestTransform * InRootTransform, false);
+		}
+	}
+	
 	bool FStaticMeshCacheAdapter::InitializeForRecord(UPrimitiveComponent* InComponent, UChaosCache* InCache)
 	{
+		InCache->Version = 1;
 		return true;
 	}
 
-	bool FStaticMeshCacheAdapter::InitializeForPlayback(UPrimitiveComponent* InComponent, UChaosCache* InCache) const
+	bool FStaticMeshCacheAdapter::InitializeForPlayback(UPrimitiveComponent* InComponent, UChaosCache* InCache, float InTime)
 	{
 #if WITH_CHAOS
 
diff --git a/Engine/Plugins/Experimental/ChaosCaching/Source/ChaosCaching/Private/Chaos/CacheManagerActor.cpp b/Engine/Plugins/Experimental/ChaosCaching/Source/ChaosCaching/Private/Chaos/CacheManagerActor.cpp
--- a/Engine/Plugins/Experimental/ChaosCaching/Source/ChaosCaching/Private/Chaos/CacheManagerActor.cpp
+++ b/Engine/Plugins/Experimental/ChaosCaching/Source/ChaosCaching/Private/Chaos/CacheManagerActor.cpp
@@ -19,6 +19,7 @@
 #include "UObject/ConstructorHelpers.h"
 #include "Widgets/Notifications/SNotificationList.h"
 #include "Kismet2/ComponentEditorUtils.h"
+#include "GeometryCollection/GeometryCollectionComponent.h"
 #endif
 
 #define LOCTEXT_NAMESPACE "ChaosCacheManager"
@@ -28,9 +29,9 @@ FName GetComponentCacheName(UPrimitiveComponent* InComponent)
 	return FName(InComponent->GetPathName(InComponent->GetWorld()));
 }
 
-void FObservedComponent::ResetRuntimeData()
+void FObservedComponent::ResetRuntimeData(const EStartMode ManagerStartMode)
 {
-	bTriggered       = StartMode == EStartMode::Timed;
+	bTriggered       = ManagerStartMode == EStartMode::Timed;
 	AbsoluteTime     = 0;
 	TimeSinceTrigger = 0;
 	Cache            = nullptr;
@@ -50,6 +51,10 @@ UPrimitiveComponent* FObservedComponent::GetComponent() const
 
 AChaosCacheManager::AChaosCacheManager(const FObjectInitializer& ObjectInitializer /*= FObjectInitializer::Get()*/)
 	: Super(ObjectInitializer)
+	, CacheMode(ECacheMode::Record)
+	, StartMode(EStartMode::Timed)
+	, bCanRecord(true)
+	, bIsSimulating(false)
 {
 	// This actor will tick, just not normally. There needs to be a tick-like event both before physics simulation
 	// and after physics simulation, we bind to some physics scene events in BeginPlay to handle this.
@@ -99,27 +104,46 @@ void AChaosCacheManager::TickActor(float DeltaTime, enum ELevelTick TickType, FA
 	}
 }
 
-#if WITH_EDITOR
 void AChaosCacheManager::PostEditChangeProperty(FPropertyChangedEvent& PropertyChangedEvent)
 {
 	Super::PostEditChangeProperty(PropertyChangedEvent);
-}
-#endif
 
-void AChaosCacheManager::SetAllMode(ECacheMode InMode)
-{
-	for(FObservedComponent& Observed : ObservedComponents)
+	FName PropertyName = (PropertyChangedEvent.Property != NULL) ? PropertyChangedEvent.Property->GetFName() : NAME_None;
+
+	if (PropertyName == GET_MEMBER_NAME_CHECKED(AChaosCacheManager, CacheMode))
 	{
-		Observed.CacheMode = InMode;
+		if ((CacheMode == ECacheMode::Record) && !CanRecord())
+		{
+			// If we're not allowed to record but have somehow been set to record mode, revert to default
+			CacheMode = ECacheMode::None;
+		}
+
+#if WITH_EDITOR
+		SetObservedComponentProperties(CacheMode);
+#endif
+		
+		if (CacheMode != ECacheMode::Record)
+		{
+			OnStartFrameChanged(StartTime);
+		}
+		else
+		{
+			OnStartFrameChanged(0.0);
+		}
+	}
+	else if (PropertyName == GET_MEMBER_NAME_CHECKED(AChaosCacheManager, StartTime)) 
+	{
+		if (CacheMode != ECacheMode::Record)
+		{
+			OnStartFrameChanged(StartTime);
+		}
 	}
 }
 
-void AChaosCacheManager::SetAllStartMode(EStartMode InMode)
+void AChaosCacheManager::SetStartTime(float InStartTime)
 {
-	for (FObservedComponent& Observed : ObservedComponents)
-	{
-		Observed.StartMode = InMode;
-	}
+	StartTime = InStartTime;
+	OnStartFrameChanged(InStartTime);
 }
 
 void AChaosCacheManager::ResetAllComponentTransforms()
@@ -184,11 +208,13 @@ void AChaosCacheManager::BeginPlay()
 
 	Super::BeginPlay();
 
+	bIsSimulating = !(CacheMode == ECacheMode::None);
+	
 	if(!CacheCollection)
 	{
 		UE_LOG(LogChaosCache, Warning, TEXT("%s has no valid cache asset. Components will revert to dynamic simulation."), *GetName());
 		
-		// without a collection the cache manager can't do anything, no reason to inialise the observed array
+		// without a collection the cache manager can't do anything, no reason to initialise the observed array
 		SetActorTickEnabled(false);
 		return;
 	}
@@ -262,7 +288,7 @@ void AChaosCacheManager::BeginPlay()
 		}
 
 		// Reset timers and last cache
-		Observed.ResetRuntimeData();
+		Observed.ResetRuntimeData(StartMode);
 
 		bool                    bRequiresRecord = false;
 		FComponentCacheAdapter* CurrAdapter     = ActiveAdapters[Index];
@@ -280,85 +306,76 @@ void AChaosCacheManager::BeginPlay()
 				SolverData->PostSolveHandle = Solver->AddPostAdvanceCallback(FSolverPostAdvance::FDelegate::CreateUObject(this, &AChaosCacheManager::HandlePostSolve, Solver));
 			}
 
-			switch(Observed.CacheMode)
+			if (CacheMode == ECacheMode::Play || CacheMode == ECacheMode::None || !CanRecord())
 			{
-				case ECacheMode::Play:
-				{
-					UChaosCache*    PlayCache = CacheCollection->FindCache(Observed.CacheName);
+				UChaosCache*    PlayCache = CacheCollection->FindCache(Observed.CacheName);
 					
-					if(PlayCache)
-					{
-						FCacheUserToken Token = PlayCache->BeginPlayback();
+				if(PlayCache)
+				{
+					FCacheUserToken Token = PlayCache->BeginPlayback();
 
-						if(Token.IsOpen() && CurrAdapter->ValidForPlayback(Comp, PlayCache))
+					if(Token.IsOpen() && CurrAdapter->ValidForPlayback(Comp, PlayCache))
+					{
+						SolverData->PlaybackIndices.Add(Index);
+						SolverData->PlaybackTickRecords.AddDefaulted();
+						SolverData->PlaybackTickRecords.Last().SetSpaceTransform(GetTransform());
+						Observed.Cache = PlayCache;
+						Observed.TickRecord.SetSpaceTransform(GetTransform());
+						Observed.TickRecord.SetLastTime(StartTime);
+						OpenPlaybackCaches.Add(TTuple<FCacheUserToken, UChaosCache*>(MoveTemp(Token), Observed.Cache));
+						CurrAdapter->InitializeForPlayback(Comp, Observed.Cache, StartTime);
+					}
+					else
+					{
+						if(Token.IsOpen())
 						{
-							SolverData->PlaybackIndices.Add(Index);
-							SolverData->PlaybackTickRecords.AddDefaulted();
-							SolverData->PlaybackTickRecords.Last().SetSpaceTransform(GetTransform());
-							Observed.Cache = PlayCache;
-							Observed.TickRecord.SetSpaceTransform(GetTransform());
-							OpenPlaybackCaches.Add(TTuple<FCacheUserToken, UChaosCache*>(MoveTemp(Token), Observed.Cache));
-							CurrAdapter->InitializeForPlayback(Comp, Observed.Cache);
+							UE_LOG(LogChaosCache,
+									Warning,
+									TEXT("Failed playback for component %s, Selected cache adapter unable to handle the cache (the cache is incompatible)"),
+									*Comp->GetPathName());
+
+							// The cache session was valid so make sure to end it
+							PlayCache->EndPlayback(Token);
 						}
-						else
+						else    // Already open for record somewhere
 						{
-							if(Token.IsOpen())
-							{
-								UE_LOG(LogChaosCache,
-									   Warning,
-									   TEXT("Failed playback for component %s, Selected cache adapter unable to handle the cache (the cache is incompatible)"),
-									   *Comp->GetPathName());
-
-								// The cache session was valid so make sure to end it
-								PlayCache->EndPlayback(Token);
-							}
-							else    // Already open for record somewhere
-							{
-								UE_LOG(LogChaosCache,
-									   Warning,
-									   TEXT("Failed playback for component %s using cache %s, cache already open for record"),
-									   *Comp->GetName(),
-									   *PlayCache->GetPathName());
-							}
-
-							++NumFailedPlaybackEntries;
+							UE_LOG(LogChaosCache,
+									Warning,
+									TEXT("Failed playback for component %s using cache %s, cache already open for record"),
+									*Comp->GetName(),
+									*PlayCache->GetPathName());
 						}
-					}
-					else
-					{
-						UE_LOG(LogChaosCache, Log, TEXT("Skipping playback for component %s, no available cache."), *Comp->GetName());
-					}
 
-					break;
+						++NumFailedPlaybackEntries;
+					}
 				}
-				case ECacheMode::Record:
+				else
 				{
-					// Make sure there's a cache available if we're going to record.
-					UPrimitiveComponent* Component = Observed.GetComponent();
-					FName CacheName = (Observed.CacheName == NAME_None) ? MakeUniqueObjectName(CacheCollection, UChaosCache::StaticClass(), "Cache") : Observed.CacheName;
+					UE_LOG(LogChaosCache, Log, TEXT("Skipping playback for component %s, no available cache."), *Comp->GetName());
+				}
+			}
+			else // CacheMode == ECacheMode::Record
+			{
+				// Make sure there's a cache available if we're going to record.
+				UPrimitiveComponent* Component = Observed.GetComponent();
+				FName CacheName = (Observed.CacheName == NAME_None) ? MakeUniqueObjectName(CacheCollection, UChaosCache::StaticClass(), "Cache") : Observed.CacheName;
 
-					UChaosCache*    RecordCache = CacheCollection->FindOrAddCache(CacheName);
-					FCacheUserToken Token       = RecordCache->BeginRecord(Observed.GetComponent(), CurrAdapter->GetGuid(), GetTransform());
+				UChaosCache*    RecordCache = CacheCollection->FindOrAddCache(CacheName);
+				FCacheUserToken Token       = RecordCache->BeginRecord(Observed.GetComponent(), CurrAdapter->GetGuid(), GetTransform());
 
-					if(Token.IsOpen())
-					{
-						SolverData->RecordIndices.Add(Index);
+				if(Token.IsOpen())
+				{
+					SolverData->RecordIndices.Add(Index);
 
-						Observed.Cache = CacheCollection->FindOrAddCache(CacheName);
+					Observed.Cache = CacheCollection->FindOrAddCache(CacheName);
 						
-						// We'll record the observed component in Cache Manager's local space.
-						Observed.TickRecord.SetSpaceTransform(GetTransform());
-						OpenRecordCaches.Add(TTuple<FCacheUserToken, UChaosCache*>(MoveTemp(Token), Observed.Cache));
-						CurrAdapter->InitializeForRecord(Component, Observed.Cache);
+					// We'll record the observed component in Cache Manager's local space.
+					Observed.TickRecord.SetSpaceTransform(GetTransform());
+					OpenRecordCaches.Add(TTuple<FCacheUserToken, UChaosCache*>(MoveTemp(Token), Observed.Cache));
+					CurrAdapter->InitializeForRecord(Component, Observed.Cache);
 
-						// Ensure we enable the actor tick to flush out the pending record writes
-						bRequiresRecord = true;
-					}
-					break;
-				}
-				default:
-				{
-					break;
+					// Ensure we enable the actor tick to flush out the pending record writes
+					bRequiresRecord = true;
 				}
 			}
 		}
@@ -386,6 +403,8 @@ void AChaosCacheManager::EndPlay(const EEndPlayReason::Type EndPlayReason)
 {
 	Super::EndPlay(EndPlayReason);
 
+	bIsSimulating = false;
+	
 	FChaosSolversModule* Module = FChaosSolversModule::GetModule();
 	check(Module);
 
@@ -484,10 +503,10 @@ void AChaosCacheManager::HandlePostSolve(Chaos::FReal InDt, Chaos::FPhysicsSolve
 
 	TickObservedComponents(Data->RecordIndices, InDt, [](UChaosCache* InCache, FObservedComponent& Observed, Chaos::FComponentCacheAdapter* InAdapter) {
 		UPrimitiveComponent* Comp = Observed.GetComponent();
-		if(ensure(Comp && InCache))
+		if (ensure(Comp && InCache))
 		{
 			// If we haven't advanced since the last record, don't push another frame
-			if(Observed.TimeSinceTrigger > InCache->GetDuration())
+			if (Observed.TimeSinceTrigger > InCache->GetDuration())
 			{
 				FPendingFrameWrite NewFrame;
 				NewFrame.Time = Observed.TimeSinceTrigger;
@@ -499,6 +518,49 @@ void AChaosCacheManager::HandlePostSolve(Chaos::FReal InDt, Chaos::FPhysicsSolve
 	});
 }
 
+void AChaosCacheManager::OnStartFrameChanged(Chaos::FReal InTime)
+{
+	if (bIsSimulating)
+	{
+		return;
+	}
+	
+	if (!CacheCollection)
+	{
+		return;
+	}
+	
+	for (int32 ObservedIdx = 0; ObservedIdx < ObservedComponents.Num(); ++ObservedIdx)
+	{
+		FObservedComponent& Observed = ObservedComponents[ObservedIdx];
+		if (UPrimitiveComponent* Comp = Observed.GetComponent())
+		{
+			if (!Observed.BestFitAdapter)
+			{
+				// Cache the best fit adapater so that we don't need to look it up each time we evaluate a start time.
+				Observed.BestFitAdapter = Chaos::FAdapterUtil::GetBestAdapterForClass(Comp->GetClass(), false);
+			}
+
+			if (!Observed.Cache)
+			{
+				Observed.Cache = CacheCollection->FindCache(Observed.CacheName);
+			}
+
+			if (!Observed.Cache || !Observed.BestFitAdapter || Observed.Cache->GetDuration()==0.0)
+			{
+				continue;
+			}
+
+			FCacheUserToken Token = Observed.Cache->BeginPlayback();
+			if (Token.IsOpen() && Observed.BestFitAdapter->ValidForPlayback(Comp, Observed.Cache))
+			{
+				Observed.BestFitAdapter->SetRestState(Comp, Observed.Cache, GetTransform(), InTime);
+				Observed.Cache->EndPlayback(Token);
+			}
+		}
+	}
+}
+
 void AChaosCacheManager::TriggerComponent(UPrimitiveComponent* InComponent)
 {
 	// #BGTODO Maybe not totally thread-safe, probably safer with an atomic or condition var rather than the bTriggered flag
@@ -506,7 +568,7 @@ void AChaosCacheManager::TriggerComponent(UPrimitiveComponent* InComponent)
 		return Test.GetComponent() == InComponent;
 	});
 
-	if(Found && Found->StartMode == EStartMode::Triggered)
+	if (Found && StartMode == EStartMode::Triggered)
 	{
 		Found->bTriggered = true;
 	}
@@ -518,7 +580,7 @@ void AChaosCacheManager::TriggerComponentByCache(FName InCacheName)
 		return Test.CacheName == InCacheName && Test.GetComponent();
 	});
 
-	if(Found && Found->StartMode == EStartMode::Triggered)
+	if (Found && StartMode == EStartMode::Triggered)
 	{
 		Found->bTriggered = true;
 	}
@@ -528,7 +590,7 @@ void AChaosCacheManager::TriggerAll()
 {
 	for(FObservedComponent& Observed : ObservedComponents)
 	{
-		if(Observed.StartMode == EStartMode::Triggered && Observed.GetComponent())
+		if (StartMode == EStartMode::Triggered && Observed.GetComponent())
 		{
 			Observed.bTriggered = true;
 		}
@@ -550,7 +612,8 @@ FObservedComponent& AChaosCacheManager::AddNewObservedComponent(UPrimitiveCompon
 
 	NewEntry.ComponentRef.PathToComponent = InComponent->GetPathName(InComponent->GetOwner());
 	NewEntry.ComponentRef.OtherActor      = InComponent->GetOwner();
-	NewEntry.CacheName                    = MakeUniqueObjectName(CacheCollection, UChaosCache::StaticClass(), "Cache");
+
+	NewEntry.CacheName                    = MakeUniqueObjectName(CacheCollection, UChaosCache::StaticClass(), FName(InComponent->GetOwner()->GetActorLabel()));
 
 	return NewEntry;
 }
@@ -593,22 +656,41 @@ void AChaosCacheManager::TickObservedComponents(const TArray<int32>& InIndices,
 		// (happens if a plugin implemented it but is no longer loaded)
 		if(Observed.bTriggered && Adapter)
 		{
-			if(Observed.CacheMode == ECacheMode::Play)
+			if (CacheMode == ECacheMode::Play)
 			{
 				Observed.TickRecord.SetDt(InDt);
 			}
 
-			if(Observed.TimedDuration == 0.0f)
+			Observed.TimeSinceTrigger += InDt;
+			InCallable(Observed.Cache, Observed, Adapter);
+		}
+	}
+}
+
+#if WITH_EDITOR
+void AChaosCacheManager::SetObservedComponentProperties(const ECacheMode& NewCacheMode)
+{
+	for (FObservedComponent& ObservedComponent : ObservedComponents)
+	{
+		if (UPrimitiveComponent* PrimComp = ObservedComponent.GetComponent())
+		{
+			if (NewCacheMode == ECacheMode::Record)
 			{
-				Observed.TimeSinceTrigger += InDt;
-				InCallable(Observed.Cache, Observed, Adapter);
+				PrimComp->BodyInstance.bSimulatePhysics = ObservedComponent.bIsSimulating;	
 			}
 			else
 			{
-				Observed.TimedDuration = FMath::Max(Observed.TimedDuration - InDt, 0.0f);
+				PrimComp->BodyInstance.bSimulatePhysics = false;
 			}
 		}
 	}
 }
+#endif
+
+AChaosCachePlayer::AChaosCachePlayer(const FObjectInitializer& ObjectInitializer /*= FObjectInitializer::Get()*/)
+	: Super(ObjectInitializer)
+{
+	bCanRecord = false;
+}
 
 #undef LOCTEXT_NAMESPACE
diff --git a/Engine/Plugins/Experimental/ChaosCaching/Source/ChaosCaching/Private/Chaos/ChaosCache.cpp b/Engine/Plugins/Experimental/ChaosCaching/Source/ChaosCaching/Private/Chaos/ChaosCache.cpp
--- a/Engine/Plugins/Experimental/ChaosCaching/Source/ChaosCaching/Private/Chaos/ChaosCache.cpp
+++ b/Engine/Plugins/Experimental/ChaosCaching/Source/ChaosCaching/Private/Chaos/ChaosCache.cpp
@@ -42,6 +42,15 @@ void UChaosCache::FlushPendingFrames()
 			{
 				// Initial write to this particle
 				PTrack.BeginOffset = NewData.Time;
+
+				// Particle will hold at end of recording.
+				PTrack.bDeactivateOnEnd = false;
+			}
+
+			if (ParticleData.bPendingDeactivate)
+			{
+				// Signals that this is the final keyframe and that the particle then deactivates.
+				PTrack.bDeactivateOnEnd = true;
 			}
 
 			// Make sure we're actually appending to the track - shouldn't be adding data from the past
@@ -302,6 +311,12 @@ FCacheEvaluationResult UChaosCache::Evaluate(const FCacheEvaluationContext& InCo
 				continue;
 			}
 
+			if (ParticleTracks[Index].TransformData.bDeactivateOnEnd && ParticleTracks[Index].TransformData.GetEndTime() < InContext.TickRecord.GetTime())
+			{
+				// Particle has deactivated so skip evaluation
+				continue;
+			}
+
 			FTransform* EvalTransform = nullptr;
 			TMap<FName, float>* EvalCurves = nullptr;
 
@@ -361,11 +376,13 @@ void UChaosCache::EvaluateSingle(int32 InIndex, FPlaybackTickRecord& InTickRecor
 	checkSlow(ParticleTracks.IsValidIndex(InIndex));
 	FPerParticleCacheData& Data = ParticleTracks[InIndex];
 
+	
 	if(OutOptTransform)
 	{
 		EvaluateTransform(Data, InTickRecord.GetTime(), *OutOptTransform);
-		(*OutOptTransform) = (*OutOptTransform) * InTickRecord.SpaceTransform;
+		(*OutOptTransform) = (*OutOptTransform) *InTickRecord.SpaceTransform;
 	}
+	
 
 	if(OutOptCurves)
 	{
diff --git a/Engine/Plugins/Experimental/ChaosCaching/Source/ChaosCaching/Public/Chaos/Adapters/CacheAdapter.h b/Engine/Plugins/Experimental/ChaosCaching/Source/ChaosCaching/Public/Chaos/Adapters/CacheAdapter.h
--- a/Engine/Plugins/Experimental/ChaosCaching/Source/ChaosCaching/Public/Chaos/Adapters/CacheAdapter.h
+++ b/Engine/Plugins/Experimental/ChaosCaching/Source/ChaosCaching/Public/Chaos/Adapters/CacheAdapter.h
@@ -117,6 +117,11 @@ namespace Chaos
 		 * Called from the game thread to perform any global setup that the adapter may need to perform.
 		 */
 		virtual void Initialize() {}
+
+		/**
+		* Called from the game thread to set rest state from an evaluated time in the cache.
+		*/
+		virtual void SetRestState(UPrimitiveComponent* InComponent, UChaosCache* InCache, const FTransform& InRootTransform, Chaos::FReal InTime) const = 0;
 		
 		/**
 		 * Called from the game thread to initialize a component and cache ready to record a cache
@@ -129,8 +134,9 @@ namespace Chaos
 		 * Called from the game thread to initialize a component and cache ready to playback a cache
 		 * @param InComponent Target component to initialize
 		 * @param InCache Target cache to initialize
+		 * @param InTime Time from cache start to evaluate initial conditions
 		 */
-		virtual bool InitializeForPlayback(UPrimitiveComponent* InComponent, UChaosCache* InCache) const = 0;
+		virtual bool InitializeForPlayback(UPrimitiveComponent* InComponent, UChaosCache* InCache, float InTime) = 0;
 
 		/**
 		 * Called by a cache observer when a component should be recorded to a cache through this adapter.
diff --git a/Engine/Plugins/Experimental/ChaosCaching/Source/ChaosCaching/Public/Chaos/Adapters/GeometryCollectionComponentCacheAdapter.h b/Engine/Plugins/Experimental/ChaosCaching/Source/ChaosCaching/Public/Chaos/Adapters/GeometryCollectionComponentCacheAdapter.h
--- a/Engine/Plugins/Experimental/ChaosCaching/Source/ChaosCaching/Public/Chaos/Adapters/GeometryCollectionComponentCacheAdapter.h
+++ b/Engine/Plugins/Experimental/ChaosCaching/Source/ChaosCaching/Public/Chaos/Adapters/GeometryCollectionComponentCacheAdapter.h
@@ -241,13 +241,14 @@ namespace Chaos
 		Chaos::FPhysicsSolver* GetComponentSolver(UPrimitiveComponent* InComponent) const override;
 		bool                   ValidForPlayback(UPrimitiveComponent* InComponent, UChaosCache* InCache) const override;
 		void				   Initialize() override;
+		void				   SetRestState(UPrimitiveComponent* InComponent, UChaosCache* InCache, const FTransform& InRootTransform, Chaos::FReal InTime) const override;
 		bool                   InitializeForRecord(UPrimitiveComponent* InComponent, UChaosCache* InCache) override;
-		bool                   InitializeForPlayback(UPrimitiveComponent* InComponent, UChaosCache* InCache) const override;
+		bool                   InitializeForPlayback(UPrimitiveComponent* InComponent, UChaosCache* InCache, float InTime) override;
 		void                   Record_PostSolve(UPrimitiveComponent* InComp, const FTransform& InRootTransform, FPendingFrameWrite& OutFrame, Chaos::FReal InTime) const override;
-		void                   Playback_PreSolve(UPrimitiveComponent*                               InComponent,
-												 UChaosCache*                                       InCache,
+		void                   Playback_PreSolve(UPrimitiveComponent*								InComponent,
+												 UChaosCache*										InCache,
 												 Chaos::FReal                                       InTime,
-												 FPlaybackTickRecord&                               TickRecord,
+												 FPlaybackTickRecord&								TickRecord,
 												 TArray<TPBDRigidParticleHandle<Chaos::FReal, 3>*>& OutUpdatedRigids) const override;
 
 		// End FComponentCacheAdapter interface
@@ -258,6 +259,12 @@ namespace Chaos
 		void HandleTrailingEvents(const Chaos::FTrailingEventData& Event);
 	
 	private:
+		TArray<int32> GatherAllBreaksUpToTime(UChaosCache* InCache, float InTime) const;
+
+		bool bHandleBreaks;
+		bool bHandleCollisions;
+		bool bHandleTrailing;
+
 		TMap<IPhysicsProxyBase*, FCachedEventData> CachedData;
 		const FBreakingDataArray* BreakingDataArray;
 		const FCollisionDataArray* CollisionDataArray;
diff --git a/Engine/Plugins/Experimental/ChaosCaching/Source/ChaosCaching/Public/Chaos/Adapters/StaticMeshComponentCacheAdapter.h b/Engine/Plugins/Experimental/ChaosCaching/Source/ChaosCaching/Public/Chaos/Adapters/StaticMeshComponentCacheAdapter.h
--- a/Engine/Plugins/Experimental/ChaosCaching/Source/ChaosCaching/Public/Chaos/Adapters/StaticMeshComponentCacheAdapter.h
+++ b/Engine/Plugins/Experimental/ChaosCaching/Source/ChaosCaching/Public/Chaos/Adapters/StaticMeshComponentCacheAdapter.h
@@ -17,13 +17,14 @@ namespace Chaos
 		FGuid                  GetGuid() const override;
 		bool                   ValidForPlayback(UPrimitiveComponent* InComponent, UChaosCache* InCache) const override;
 		Chaos::FPhysicsSolver* GetComponentSolver(UPrimitiveComponent* InComponent) const override;
+		void				   SetRestState(UPrimitiveComponent* InComponent, UChaosCache* InCache, const FTransform& InRootTransform, Chaos::FReal InTime) const override;
 		bool                   InitializeForRecord(UPrimitiveComponent* InComponent, UChaosCache* InCache) override;
-		bool                   InitializeForPlayback(UPrimitiveComponent* InComponent, UChaosCache* InCache) const override;
+		bool                   InitializeForPlayback(UPrimitiveComponent* InComponent, UChaosCache* InCache, float InTime) override;
 		void                   Record_PostSolve(UPrimitiveComponent* InComp, const FTransform& InRootTransform, FPendingFrameWrite& OutFrame, Chaos::FReal InTime) const override;
-		void                   Playback_PreSolve(UPrimitiveComponent*                               InComponent,
-												 UChaosCache*                                       InCache,
+		void                   Playback_PreSolve(UPrimitiveComponent* InComponent,
+												 UChaosCache*										InCache,
 												 Chaos::FReal                                       InTime,
-												 FPlaybackTickRecord&                               TickRecord,
+												 FPlaybackTickRecord&								TickRecord,
 												 TArray<TPBDRigidParticleHandle<Chaos::FReal, 3>*>& OutUpdatedRigids) const override;
 
 	private:
diff --git a/Engine/Plugins/Experimental/ChaosCaching/Source/ChaosCaching/Public/Chaos/CacheManagerActor.h b/Engine/Plugins/Experimental/ChaosCaching/Source/ChaosCaching/Public/Chaos/CacheManagerActor.h
--- a/Engine/Plugins/Experimental/ChaosCaching/Source/ChaosCaching/Public/Chaos/CacheManagerActor.h
+++ b/Engine/Plugins/Experimental/ChaosCaching/Source/ChaosCaching/Public/Chaos/CacheManagerActor.h
@@ -22,8 +22,8 @@ class FComponentCacheAdapter;
 UENUM()
 enum class ECacheMode : uint8
 {
-	None,
-	Play,
+	None		UMETA(DisplayName = "Static Pose"),
+	Play,		
 	Record
 };
 
@@ -44,13 +44,12 @@ struct CHAOSCACHING_API FObservedComponent
 {
 	GENERATED_BODY()
 
-	FObservedComponent()
+		FObservedComponent()
 		: CacheName(NAME_None)
-		, CacheMode(ECacheMode::None)
-		, StartMode(EStartMode::Timed)
-		, TimedDuration(0.0f)
+		, bIsSimulating(true)
+		, Cache(nullptr)
+		, BestFitAdapter(nullptr)
 	{
-		ResetRuntimeData();
 	}
 
 	/** Unique name for the cache, used as a key into the cache collection */
@@ -61,27 +60,12 @@ struct CHAOSCACHING_API FObservedComponent
 	UPROPERTY(EditAnywhere, Category = "Caching", meta = (UseComponentPicker, AllowAnyActor))
 	FComponentReference ComponentRef;
 
-	/** How to use the cache - playback or record */
-	UPROPERTY(EditAnywhere, Category = "Caching")
-	ECacheMode CacheMode;
-
-	/**
-	 * How to trigger the cache record or playback, timed will start counting at BeginPlay, Triggered will begin
-	 * counting from when the owning cache manager is requested to trigger the cache action
-	 * @see AChaosCacheManager::TriggerObservedComponent
-	 */
+	/** Capture of the initial state of the component before cache manager takes control. */
 	UPROPERTY(EditAnywhere, Category = "Caching")
-	EStartMode StartMode;
-
-	/**
-	 * The time after BeginPlay or a call to AChaosCacheManager::TriggerObservedComponent to wait before beginning
-	 * the playback or recording of the component
-	 */
-	UPROPERTY(EditAnywhere, Category = "Caching")
-	float TimedDuration;
+	bool bIsSimulating;				  
 
 	/** Prepare runtime tick data for a new run */
-	void ResetRuntimeData();
+	void ResetRuntimeData(const EStartMode ManagerStartMode);
 
 	/** Gets the component from the internal component ref */
 	UPrimitiveComponent* GetComponent();
@@ -95,6 +79,8 @@ private:
 	Chaos::FReal TimeSinceTrigger;    // Time since our trigger happened
 	UChaosCache* Cache;               // Cache to play - picked up in BeginPlay on the manager.
 	FPlaybackTickRecord TickRecord;   // Tick record to track where we are in playback
+
+	Chaos::FComponentCacheAdapter* BestFitAdapter; // Cache a pointer to the best adapter, required for cache random access handling.
 };
 
 struct FPerSolverData
@@ -129,6 +115,25 @@ public:
 	UPROPERTY(EditInstanceOnly, Category = "Caching")
 	UChaosCacheCollection* CacheCollection;
 
+	/** How to use the cache - playback or record */
+	UPROPERTY(EditAnywhere, Category = "Caching")
+	ECacheMode CacheMode;
+
+	/**
+	* How to trigger the cache record or playback, timed will start counting at BeginPlay, Triggered will begin
+	* counting from when the owning cache manager is requested to trigger the cache action
+	* @see AChaosCacheManager::TriggerObservedComponent
+	*/
+	UPROPERTY(EditAnywhere, Category = "Caching")
+	EStartMode StartMode;
+
+	/**
+	* Defines the (random access) time that represents the rest pose of the components managed by this cache.
+	* When in Play mode, the components are set to the state provided by the caches at this evaluated time.
+	*/
+	UPROPERTY(EditAnywhere, Interp, BlueprintReadWrite, Category = "Caching", meta=(SequencerTrackClass="MovieSceneFloatTrack"))
+	float StartTime;
+
 	/** AActor interface */
 	void TickActor(float DeltaTime, enum ELevelTick TickType, FActorTickFunction& ThisTickFunction) override;
 	/** end AActor interface */
@@ -136,23 +141,14 @@ public:
 	/** UObject interface */
 #if WITH_EDITOR
 	friend class IChaosCachingEditorPlugin;
-	void PostEditChangeProperty(FPropertyChangedEvent& PropertyChangedEvent) override;
 #endif
-	/** end UObject interface */
 
-	/**
-	 * Sets the playback mode of every observed component to the specified mode
-	 * @param InMode Mode to set
-	 */
-	UFUNCTION(BlueprintCallable, Category = "Caching")
-	void SetAllMode(ECacheMode InMode);
+	void PostEditChangeProperty(FPropertyChangedEvent& PropertyChangedEvent) override;
+	/** end UObject interface */
 
-	/**
-	 * Sets the start mode of every observed component to the specified mode
-	 * @param InMode Mode to set
-	 */
-	UFUNCTION(BlueprintCallable, Category = "Caching")
-	void SetAllStartMode(EStartMode InMode);
+	/** Expose StartTime property to Sequencer. GetStartTime will be called on keys. */
+	UFUNCTION(CallInEditor)
+	void SetStartTime(float InStartTime);
 
 	/** 
 	 * Resets all components back to the world space transform they had when the cache for them was originally recorded
@@ -177,6 +173,9 @@ public:
 	void SelectComponent(int32 InIndex);
 #endif
 
+	/** Returns true if this cache manager is allowed to record caches. */
+	bool CanRecord() const { return bCanRecord; }
+
 protected:
 	/** AActor interface */
 	void BeginPlay() override;
@@ -192,6 +191,9 @@ protected:
 	/** Handles physics thread post-solve (record data for components under record) */
 	void HandlePostSolve(Chaos::FReal InDt, Chaos::FPhysicsSolver* InSolver);
 
+	/** Evaluates and sets state for all observed components at the specified time. */
+	void OnStartFrameChanged(Chaos::FReal InT);
+
 	/**
 	 * Triggers a component to play or record.
 	 * If the cache manager has an observed component entry for InComponent and it is a triggered entry
@@ -219,7 +221,17 @@ protected:
 	FObservedComponent& FindOrAddObservedComponent(UPrimitiveComponent* InComponent);
 	void ClearObservedComponents();
 
+	// Determines if the actor is allowed to record a cache.
+	bool bCanRecord;
+
+	// true if observed components are actively simulating, dynamically or kinematically.
+	bool bIsSimulating;
+
 private:
+#if WITH_EDITOR
+	void SetObservedComponentProperties(const ECacheMode& NewCacheMode);
+#endif
+	
 	friend class UActorFactoryCacheManager; // Allows the actor factory to set up the observed list. See UActorFactoryCacheManager::PostSpawnActor
 
 	using FTickObservedFunction = TUniqueFunction<void(UChaosCache*, FObservedComponent&, Chaos::FComponentCacheAdapter*)>;
@@ -247,3 +259,13 @@ private:
 	TArray<TTuple<FCacheUserToken, UChaosCache*>> OpenRecordCaches;
 	TArray<TTuple<FCacheUserToken, UChaosCache*>> OpenPlaybackCaches;
 };
+
+UCLASS(Experimental)
+class CHAOSCACHING_API AChaosCachePlayer : public AChaosCacheManager
+{
+	GENERATED_BODY()
+
+public:
+	AChaosCachePlayer(const FObjectInitializer& ObjectInitializer = FObjectInitializer::Get());
+
+};
diff --git a/Engine/Plugins/Experimental/ChaosCaching/Source/ChaosCaching/Public/Chaos/ChaosCache.h b/Engine/Plugins/Experimental/ChaosCaching/Source/ChaosCaching/Public/Chaos/ChaosCache.h
--- a/Engine/Plugins/Experimental/ChaosCaching/Source/ChaosCaching/Public/Chaos/ChaosCache.h
+++ b/Engine/Plugins/Experimental/ChaosCaching/Source/ChaosCaching/Public/Chaos/ChaosCache.h
@@ -28,6 +28,10 @@ struct FParticleTransformTrack
 	UPROPERTY()
 	float BeginOffset;
 
+	/** If this flag is set true, the particle represented by this track deactives on the final keyframe. */
+	UPROPERTY()
+	bool bDeactivateOnEnd;
+
 	/**
 	 * The above raw track is just the key data and doesn't know at which time those keys are placed, this is
 	 * a list of the timestamps for each entry in TransformTrack
@@ -101,6 +105,11 @@ struct FPlaybackTickRecord
 		LastEventPerTrack.Reset();
 	}
 
+	void SetLastTime(float InTime)
+	{
+		LastTime = InTime;
+	}
+
 	float GetTime() const
 	{
 		return LastTime + CurrentDt;
@@ -161,6 +170,7 @@ struct FPendingParticleWrite
 {
 	int32                       ParticleIndex;
 	FTransform                  PendingTransform;
+	bool						bPendingDeactivate = false;
 	TArray<TPair<FName, float>> PendingCurveData;
 };
 
@@ -338,6 +348,14 @@ public:
 	UPROPERTY()
 	TMap<FName, FRichCurve> CurveData;
 
+	// Version for introducing new features to caches. If cache recording changes content or format of any caches,
+	// the version should be incremented and cache playback should accommodate older cache versions.
+	// The cache version is during InitializeForRecord in the appropriate adapter.
+	// 0
+	// 1 : Removed MassToLocal transform for GeometryCollection caches
+	UPROPERTY()
+	int32 Version;
+
 	template<typename T>
 	FCacheEventTrack& FindOrAddEventTrack(FName InName)
 	{
diff --git a/Engine/Plugins/Experimental/ChaosCaching/Source/ChaosCachingEditor/ChaosCachingEditor.Build.cs b/Engine/Plugins/Experimental/ChaosCaching/Source/ChaosCachingEditor/ChaosCachingEditor.Build.cs
--- a/Engine/Plugins/Experimental/ChaosCaching/Source/ChaosCachingEditor/ChaosCachingEditor.Build.cs
+++ b/Engine/Plugins/Experimental/ChaosCaching/Source/ChaosCachingEditor/ChaosCachingEditor.Build.cs
@@ -24,7 +24,8 @@ namespace UnrealBuildTool.Rules
 					"BlueprintGraph",
 					"ToolMenus",
 					"PhysicsCore",
-					"ChaosCaching"
+					"ChaosCaching",
+					"GeometryCollectionEngine"
 				});
 
 			SetupModulePhysicsSupport(Target);
diff --git a/Engine/Plugins/Experimental/ChaosCaching/Source/ChaosCachingEditor/Private/Chaos/ActorFactoryCacheManager.cpp b/Engine/Plugins/Experimental/ChaosCaching/Source/ChaosCachingEditor/Private/Chaos/ActorFactoryCacheManager.cpp
--- a/Engine/Plugins/Experimental/ChaosCaching/Source/ChaosCachingEditor/Private/Chaos/ActorFactoryCacheManager.cpp
+++ b/Engine/Plugins/Experimental/ChaosCaching/Source/ChaosCachingEditor/Private/Chaos/ActorFactoryCacheManager.cpp
@@ -6,6 +6,7 @@
 #include "Chaos/CacheManagerActor.h"
 #include "Chaos/ChaosCache.h"
 #include "Components/PrimitiveComponent.h"
+#include "GeometryCollection/GeometryCollectionComponent.h"
 #include "Kismet2/ComponentEditorUtils.h"
 
 #define LOCTEXT_NAMESPACE "CacheManagerActorFactory"
@@ -13,7 +14,7 @@
 UActorFactoryCacheManager::UActorFactoryCacheManager()
 {
 	DisplayName            = LOCTEXT("DisplayName", "Chaos Cache Manager");
-	NewActorClass          = AChaosCacheManager::StaticClass();
+	NewActorClass          = AChaosCachePlayer::StaticClass();
 	bUseSurfaceOrientation = false;
 }
 
@@ -30,12 +31,14 @@ bool UActorFactoryCacheManager::CanCreateActorFrom(const FAssetData& AssetData,
 
 void UActorFactoryCacheManager::PostSpawnActor(UObject* Asset, AActor* NewActor)
 {
-	AChaosCacheManager*    Manager    = Cast<AChaosCacheManager>(NewActor);
+	AChaosCachePlayer*    Manager    = Cast<AChaosCachePlayer>(NewActor);
 	UChaosCacheCollection* Collection = Cast<UChaosCacheCollection>(Asset);
 	
 	// The cachemanager exists now, start adding our spawnables
 	if (Manager && Collection)
 	{
+		Manager->CacheCollection = Collection;
+
 		if (!Manager->bIsEditorPreviewActor)
 		{
 			if (UWorld* World = Manager->GetWorld())
@@ -43,7 +46,6 @@ void UActorFactoryCacheManager::PostSpawnActor(UObject* Asset, AActor* NewActor)
 				// Remove any pre-existing components
 				Manager->ClearObservedComponents();
 				
-				Manager->CacheCollection = Collection;
 				for (UChaosCache* Cache : Collection->GetCaches())
 				{
 					if (!Cache)
@@ -66,18 +68,24 @@ void UActorFactoryCacheManager::PostSpawnActor(UObject* Asset, AActor* NewActor)
 						NewComponent->SetRelativeTransform(Template.ComponentTransform);
 						Manager->AddInstanceComponent(NewComponent);
 						NewComponent->RegisterComponent();
+						
+						if (UGeometryCollectionComponent* GCComp = Cast<UGeometryCollectionComponent>(NewComponent))
+						{
+							GCComp->ObjectType = EObjectStateTypeEnum::Chaos_Object_Kinematic;
+						}
+
+						NewComponent->BodyInstance.bSimulatePhysics = false;
 
 						FObservedComponent& Observed = Manager->AddNewObservedComponent(NewComponent);
-						// AddNewObservedComponent will have given this a unique name as if it was going to build a new cache, override to the actual cache name
 						Observed.CacheName = Cache->GetFName();
-						Observed.CacheMode = ECacheMode::Play;
-						Observed.StartMode = EStartMode::Timed;
-						Observed.TimedDuration = 0.0f;
 					}
 				}
 			}
 		}
-	}
+		// Caches placed in this fashion are for Playback, not recording.
+		Manager->CacheMode = ECacheMode::Play;
+		Manager->OnStartFrameChanged(0.0);
+	}	
 }
 
 UObject* UActorFactoryCacheManager::GetAssetFromActorInstance(AActor* ActorInstance)
diff --git a/Engine/Plugins/Experimental/ChaosCaching/Source/ChaosCachingEditor/Private/Chaos/CacheEditorCommands.cpp b/Engine/Plugins/Experimental/ChaosCaching/Source/ChaosCachingEditor/Private/Chaos/CacheEditorCommands.cpp
--- a/Engine/Plugins/Experimental/ChaosCaching/Source/ChaosCachingEditor/Private/Chaos/CacheEditorCommands.cpp
+++ b/Engine/Plugins/Experimental/ChaosCaching/Source/ChaosCachingEditor/Private/Chaos/CacheEditorCommands.cpp
@@ -7,10 +7,6 @@
 void FCachingEditorCommands::RegisterCommands()
 {
 	UI_COMMAND(CreateCacheManager, "Create Cache Manager", "Adds a cache manager to observe compatible components in the selection set.", EUserInterfaceActionType::Button, FInputChord());
-	UI_COMMAND(SetManagerRecordAll, "Set Record All", "Sets selected cache managers to record all of their observed components.", EUserInterfaceActionType::Button, FInputChord());
-	UI_COMMAND(SetManagerPlayAll, "Set Play All", "Sets selected cache managers to play all of their observed components.", EUserInterfaceActionType::Button, FInputChord());
-	UI_COMMAND(SetManagerTriggerAll, "Set Trigger All", "Sets selected cache managers to set all of their observed components to triggered.", EUserInterfaceActionType::Button, FInputChord());
-	UI_COMMAND(SetManagerTimedAll, "Set Timed All", "Sets selected cache managers to set all of their observed components to timed.", EUserInterfaceActionType::Button, FInputChord());
 }
 
 #undef LOCTEXT_NAMESPACE
diff --git a/Engine/Plugins/Experimental/ChaosCaching/Source/ChaosCachingEditor/Private/Chaos/CacheManagerCustomization.cpp b/Engine/Plugins/Experimental/ChaosCaching/Source/ChaosCachingEditor/Private/Chaos/CacheManagerCustomization.cpp
--- a/Engine/Plugins/Experimental/ChaosCaching/Source/ChaosCachingEditor/Private/Chaos/CacheManagerCustomization.cpp
+++ b/Engine/Plugins/Experimental/ChaosCaching/Source/ChaosCachingEditor/Private/Chaos/CacheManagerCustomization.cpp
@@ -12,32 +12,6 @@
 
 #define LOCTEXT_NAMESPACE "CacheManagerDetails"
 
-FReply OnClickSetAllButton(TArray<AChaosCacheManager*> Managers, ECacheMode NewMode)
-{
-	for(AChaosCacheManager* Manager : Managers)
-	{
-		if(Manager)
-		{
-			Manager->SetAllMode(NewMode);
-		}
-	}
-
-	return FReply::Handled();
-}
-
-FReply OnClickSetAllStartButton(TArray<AChaosCacheManager*> Managers, EStartMode NewMode)
-{
-	for (AChaosCacheManager* Manager : Managers)
-	{
-		if (Manager)
-		{
-			Manager->SetAllStartMode(NewMode);
-		}
-	}
-
-	return FReply::Handled();
-}
-
 FReply OnClickResetTransforms(TArray<AChaosCacheManager*> Managers)
 {
 	for(AChaosCacheManager* Manager : Managers)
@@ -61,11 +35,14 @@ void FCacheManagerDetails::CustomizeDetails(IDetailLayoutBuilder& DetailBuilder)
 	TArray<TWeakObjectPtr<UObject>> Selected = DetailBuilder.GetSelectedObjects();
 	TArray<AChaosCacheManager*> CacheManagersInSelection;
 
+	bool bCanRecord = true;
+
 	for(TWeakObjectPtr<UObject> Ptr : Selected)
 	{
 		if(AChaosCacheManager* Manager = Cast<AChaosCacheManager>(Ptr.Get()))
 		{
 			CacheManagersInSelection.Add(Manager);
+			bCanRecord &= Manager->CanRecord();
 		}
 	}
 
@@ -74,96 +51,36 @@ void FCacheManagerDetails::CustomizeDetails(IDetailLayoutBuilder& DetailBuilder)
 		return;
 	}
 
-	IDetailCategoryBuilder& CachingCategory = DetailBuilder.EditCategory("Caching");
-	FDetailWidgetRow&       SetAllRow       = CachingCategory.AddCustomRow(FText::GetEmpty());
-
-	SetAllRow.NameContent()
-	[
-		SNew(STextBlock)
-		.Font(DetailBuilder.GetDetailFont())
-		.Text(LOCTEXT("SetAllLabel", "Set All"))
-	];
+	SetCacheModeOptions(bCanRecord);
 
-	SetAllRow.ValueContent()
-	.MinDesiredWidth(300.0f)
-	[
-		SNew(SHorizontalBox)
-		+ SHorizontalBox::Slot()
-		.Padding(0.0f, 0.0f, 3.0f, 0.0f)
-		[
-			SNew(SButton)
-			.HAlign(HAlign_Center)
-			.VAlign(VAlign_Center)
-			.OnClicked_Static(&OnClickSetAllButton, CacheManagersInSelection, ECacheMode::Record)
-			[
-				SNew(STextBlock)
-				.Text(LOCTEXT("SetAllRecord", "Record"))
-			]
-		]
-		+ SHorizontalBox::Slot()
-		.Padding(0.0f, 0.0f, 3.0f, 0.0f)
-		[
-			SNew(SButton)
-			.HAlign(HAlign_Center)
-			.VAlign(VAlign_Center)
-			.OnClicked_Static(&OnClickSetAllButton, CacheManagersInSelection, ECacheMode::Play)
-			[
-				SNew(STextBlock)
-				.Text(LOCTEXT("SetAllPlay", "Play"))
-			]
-		]
-		+ SHorizontalBox::Slot()
-		.Padding(0.0f, 0.0f, 0.0f, 0.0f)
-		[
-			SNew(SButton)
-			.HAlign(HAlign_Center)
-			.VAlign(VAlign_Center)
-			.OnClicked_Static(&OnClickSetAllButton, CacheManagersInSelection, ECacheMode::None)
-			[
-				SNew(STextBlock)
-				.Text(LOCTEXT("SetAllNone", "None"))
-			]
-		]
-	];
+	// Hide the Record mode
+	CacheModeHandle = DetailBuilder.GetProperty(GET_MEMBER_NAME_CHECKED(AChaosCacheManager, CacheMode), AChaosCacheManager::StaticClass());
+	check(CacheModeHandle->IsValidHandle());
 
-	FDetailWidgetRow& SetTriggerRow = CachingCategory.AddCustomRow(FText::GetEmpty());
-	SetTriggerRow.NameContent()
-	[
-		SNew(STextBlock)
-		.Font(DetailBuilder.GetDetailFont())
-		.Text(LOCTEXT("SetTriggerLabel", "Set Triggers"))
-	];
-
-	SetTriggerRow.ValueContent()
-	.MinDesiredWidth(300.0f)
-	[
-		SNew(SHorizontalBox)
-		+ SHorizontalBox::Slot()
-		.Padding(0.0f, 0.0f, 3.0f, 0.0f)
+	DetailBuilder.EditDefaultProperty(CacheModeHandle)->CustomWidget()
+		.NameContent()
 		[
-			SNew(SButton)
-			.HAlign(HAlign_Center)
-			.VAlign(VAlign_Center)
-			.OnClicked_Static(&OnClickSetAllStartButton, CacheManagersInSelection, EStartMode::Triggered)
-			[
-				SNew(STextBlock)
-				.Text(LOCTEXT("SetAllTriggered", "Triggered"))
-			]
+			CacheModeHandle->CreatePropertyNameWidget()
 		]
-		+ SHorizontalBox::Slot()
-		.Padding(0.0f, 0.0f, 3.0f, 0.0f)
+		.ValueContent()
+		.MaxDesiredWidth(250.0f)			
 		[
-			SNew(SButton)
-			.HAlign(HAlign_Center)
-			.VAlign(VAlign_Center)
-			.OnClicked_Static(&OnClickSetAllStartButton, CacheManagersInSelection, EStartMode::Timed)
+			SNew(SComboBox<TSharedPtr<FText>>)
+			.OptionsSource(&CacheModeComboList)
+			.OnGenerateWidget(this, &FCacheManagerDetails::GenerateCacheModeWidget)
+			.OnSelectionChanged(this, &FCacheManagerDetails::OnCacheModeChanged)
+			.InitiallySelectedItem(GetCurrentCacheMode())	
 			[
 				SNew(STextBlock)
-				.Text(LOCTEXT("SetAllTimed", "Timed"))
-			]
-		]
-	];
+				.Text(this, &FCacheManagerDetails::GetCacheModeComboBoxContent)
+				.Font(IDetailLayoutBuilder::GetDetailFont())
+				.ColorAndOpacity(this, &FCacheManagerDetails::GetCacheModeTextColor)
+				.ToolTipText(this, &FCacheManagerDetails::GetCacheModeComboBoxContent)
+			] 
+		];
 
+	
+	IDetailCategoryBuilder& CachingCategory = DetailBuilder.EditCategory("Caching");
 	FDetailWidgetRow& ResetPositionsRow = CachingCategory.AddCustomRow(FText::GetEmpty());
 
 	ResetPositionsRow.ValueContent()
@@ -180,10 +97,85 @@ void FCacheManagerDetails::CustomizeDetails(IDetailLayoutBuilder& DetailBuilder)
 	];
 }
 
+void FCacheManagerDetails::SetCacheModeOptions(bool bRecord)
+{
+	CacheModeComboList.Empty(3);
+	CacheModeColor.Empty(3);
+
+	CacheModeComboList.Add(MakeShareable(new FText(LOCTEXT("StaticPose","Static Pose"))));
+	CacheModeColor.Add(FSlateColor::UseForeground());
+
+	CacheModeComboList.Add(MakeShareable(new FText(LOCTEXT("Play","Play"))));
+	CacheModeColor.Add(FSlateColor::UseForeground());
+
+	if (bRecord)
+	{
+		CacheModeComboList.Add(MakeShareable(new FText(LOCTEXT("Record", "RECORD"))));
+		CacheModeColor.Add(FLinearColor::Red);
+	}
+}
+
 void FCacheManagerDetails::GenerateCacheArrayElementWidget(TSharedRef<IPropertyHandle> PropertyHandle, int32 ArrayIndex, IDetailChildrenBuilder& ChildrenBuilder, IDetailLayoutBuilder* DetailLayout)
 {
 }
 
+TSharedRef<SWidget> FCacheManagerDetails::GenerateCacheModeWidget(TSharedPtr<FText> InItem)
+{
+	FSlateColor ItemColor = FSlateColor::UseForeground();
+	if (InItem->EqualTo(LOCTEXT("Record","RECORD"), ETextComparisonLevel::Default))
+	{
+		ItemColor = FLinearColor::Red;
+	}
+
+	return
+		SNew(STextBlock)
+		.Text(*InItem)
+		.ColorAndOpacity(ItemColor)
+		.Font(IDetailLayoutBuilder::GetDetailFont());
+}
+
+void FCacheManagerDetails::OnCacheModeChanged(TSharedPtr<FText> NewSelection, ESelectInfo::Type SelectInfo)
+{
+	FText const NewCacheModeName = *NewSelection.Get();
+	for (int32 Idx = 0; Idx < CacheModeComboList.Num(); ++Idx)
+	{
+		if (NewCacheModeName.EqualTo(*CacheModeComboList[Idx].Get(), ETextComparisonLevel::Default))
+		{
+			CacheModeHandle->SetValue(static_cast<uint8>(Idx));
+			return;
+		}
+	}
+	
+	// should not get here
+	check(false);
+}
+
+FText FCacheManagerDetails::GetCacheModeComboBoxContent() const
+{
+	return *CacheModeComboList[GetClampedCacheModeIndex()].Get();
+}
+
+TSharedPtr<FText> FCacheManagerDetails::GetCurrentCacheMode() const
+{
+	return CacheModeComboList[GetClampedCacheModeIndex()];
+}
+
+FSlateColor FCacheManagerDetails::GetCacheModeTextColor() const
+{
+	return CacheModeColor[GetClampedCacheModeIndex()];
+}
+
+int32  FCacheManagerDetails::GetClampedCacheModeIndex() const
+{
+	// A CacheManager may have a CacheMode (ie. Record) that we won't provide option for,
+	// since other Managers without that option may also be selected.
+	uint8 CacheMode;
+	CacheModeHandle->GetValue(CacheMode);
+
+	return FMath::Min(static_cast<int32>(CacheMode), CacheModeComboList.Num() - 1);
+}
+
+
 TSharedRef<IPropertyTypeCustomization> FObservedComponentDetails::MakeInstance()
 {
 	return MakeShareable(new FObservedComponentDetails);
diff --git a/Engine/Plugins/Experimental/ChaosCaching/Source/ChaosCachingEditor/Private/Chaos/ChaosCachingEditorPlugin.cpp b/Engine/Plugins/Experimental/ChaosCaching/Source/ChaosCachingEditor/Private/Chaos/ChaosCachingEditorPlugin.cpp
--- a/Engine/Plugins/Experimental/ChaosCaching/Source/ChaosCachingEditor/Private/Chaos/ChaosCachingEditorPlugin.cpp
+++ b/Engine/Plugins/Experimental/ChaosCaching/Source/ChaosCachingEditor/Private/Chaos/ChaosCachingEditorPlugin.cpp
@@ -102,24 +102,6 @@ void IChaosCachingEditorPlugin::RegisterCachingSubMenu(UToolMenu* InMenu, FToolM
 									  FCanExecuteAction(),
 									  FIsActionChecked(),
 									  FIsActionButtonVisible::CreateStatic(&IsCreateCacheManagerVisible)));
-
-	InSection->AddMenuEntry("SetRecordAll",
-							LOCTEXT("MenuItem_SetRecordAll", "Set All Record"),
-							LOCTEXT("MenuItem_SetRecordAll_ToolTip", "Sets selected cache managers to record all of their observed components."),
-							FSlateIcon(),
-							FUIAction(FExecuteAction::CreateRaw(this, &IChaosCachingEditorPlugin::OnSetAllRecord),
-									  FCanExecuteAction(),
-									  FIsActionChecked(),
-									  FIsActionButtonVisible::CreateStatic(&IsSetAllRecordVisible)));
-
-	InSection->AddMenuEntry("SetPlayAll",
-							LOCTEXT("MenuItem_SetPlayAll", "Set All Play"),
-							LOCTEXT("MenuItem_SetPlayAll_ToolTip", "Sets selected cache managers to playback all of their observed components."),
-							FSlateIcon(),
-							FUIAction(FExecuteAction::CreateRaw(this, &IChaosCachingEditorPlugin::OnSetAllPlay),
-									  FCanExecuteAction(),
-									  FIsActionChecked(),
-									  FIsActionButtonVisible::CreateStatic(&IsSetAllPlayVisible)));
 }
 
 void IChaosCachingEditorPlugin::OnCreateCacheManager()
@@ -179,6 +161,7 @@ void IChaosCachingEditorPlugin::OnCreateCacheManager()
 				if(!Existing)
 				{
 					FObservedComponent& NewEntry = Manager->AddNewObservedComponent(PrimitiveComp);
+					NewEntry.bIsSimulating = PrimitiveComp->BodyInstance.bSimulatePhysics;
 				}
 			}
 		}
@@ -198,6 +181,9 @@ void IChaosCachingEditorPlugin::OnCreateCacheManager()
 				Manager->CacheCollection = NewAsset;
 			}
 		}
+
+		// Initialize observed components according to mode
+		Manager->SetObservedComponentProperties(Manager->CacheMode);
 	}
 }
 
@@ -244,22 +230,6 @@ bool IsCreateCacheManagerVisible()
 	return false;
 }
 
-void IChaosCachingEditorPlugin::OnSetAllPlay()
-{
-	USelection* SelectedActors = GEditor->GetSelectedActors();
-
-	TArray<AChaosCacheManager*> CacheManagers;
-	SelectedActors->GetSelectedObjects<AChaosCacheManager>(CacheManagers);
-
-	for(AChaosCacheManager* Manager : CacheManagers)
-	{
-		if(Manager)
-		{
-			Manager->SetAllMode(ECacheMode::Play);
-		}
-	}
-}
-
 template<typename T>
 bool SelectionContains()
 {
@@ -278,22 +248,6 @@ bool IsSetAllPlayVisible()
 	return SelectionContains<AChaosCacheManager>();
 }
 
-void IChaosCachingEditorPlugin::OnSetAllRecord()
-{
-	USelection* SelectedActors = GEditor->GetSelectedActors();
-
-	TArray<AChaosCacheManager*> CacheManagers;
-	SelectedActors->GetSelectedObjects<AChaosCacheManager>(CacheManagers);
-
-	for(AChaosCacheManager* Manager : CacheManagers)
-	{
-		if(Manager)
-		{
-			Manager->SetAllMode(ECacheMode::Record);
-		}
-	}
-}
-
 bool IsSetAllRecordVisible()
 {
 	return SelectionContains<AChaosCacheManager>();
diff --git a/Engine/Plugins/Experimental/ChaosCaching/Source/ChaosCachingEditor/Public/Chaos/CacheEditorCommands.h b/Engine/Plugins/Experimental/ChaosCaching/Source/ChaosCachingEditor/Public/Chaos/CacheEditorCommands.h
--- a/Engine/Plugins/Experimental/ChaosCaching/Source/ChaosCachingEditor/Public/Chaos/CacheEditorCommands.h
+++ b/Engine/Plugins/Experimental/ChaosCaching/Source/ChaosCachingEditor/Public/Chaos/CacheEditorCommands.h
@@ -20,8 +20,4 @@ public:
 	CHAOSCACHINGEDITOR_API virtual void RegisterCommands();
 
 	TSharedPtr<FUICommandInfo> CreateCacheManager;
-	TSharedPtr<FUICommandInfo> SetManagerRecordAll;
-	TSharedPtr<FUICommandInfo> SetManagerPlayAll;
-	TSharedPtr<FUICommandInfo> SetManagerTriggerAll;
-	TSharedPtr<FUICommandInfo> SetManagerTimedAll;
 };
diff --git a/Engine/Plugins/Experimental/ChaosCaching/Source/ChaosCachingEditor/Public/Chaos/CacheManagerCustomization.h b/Engine/Plugins/Experimental/ChaosCaching/Source/ChaosCachingEditor/Public/Chaos/CacheManagerCustomization.h
--- a/Engine/Plugins/Experimental/ChaosCaching/Source/ChaosCachingEditor/Public/Chaos/CacheManagerCustomization.h
+++ b/Engine/Plugins/Experimental/ChaosCaching/Source/ChaosCachingEditor/Public/Chaos/CacheManagerCustomization.h
@@ -5,6 +5,7 @@
 #include "IDetailCustomization.h"
 #include "Templates/SharedPointer.h"
 #include "IPropertyTypeCustomization.h"
+#include "Widgets/Input/SComboBox.h"
 
 class IDetailLayoutBuilder;
 class IPropertyHandle;
@@ -17,8 +18,19 @@ public:
 	void CustomizeDetails(IDetailLayoutBuilder& DetailBuilder) override;
 
 private:
+	void SetCacheModeOptions(bool bRecord);
+	int32 GetClampedCacheModeIndex() const;
+	FSlateColor GetCacheModeTextColor() const;
 
 	void GenerateCacheArrayElementWidget(TSharedRef<IPropertyHandle> PropertyHandle, int32 ArrayIndex, IDetailChildrenBuilder& ChildrenBuilder, IDetailLayoutBuilder* DetailLayout);
+	TSharedRef<SWidget> GenerateCacheModeWidget(TSharedPtr<FText> InItem);
+	void OnCacheModeChanged(TSharedPtr<FText> NewSelection, ESelectInfo::Type SelectInfo);
+	FText GetCacheModeComboBoxContent() const;
+	TSharedPtr<FText> GetCurrentCacheMode() const;
+
+	TSharedPtr<IPropertyHandle> CacheModeHandle;
+	TArray<TSharedPtr<FText>> CacheModeComboList;
+	TArray<FSlateColor> CacheModeColor;
 };
 
 class FObservedComponentDetails : public IPropertyTypeCustomization
diff --git a/Engine/Plugins/Experimental/ChaosCaching/Source/ChaosCachingEditor/Public/Chaos/ChaosCachingEditorPlugin.h b/Engine/Plugins/Experimental/ChaosCaching/Source/ChaosCachingEditor/Public/Chaos/ChaosCachingEditorPlugin.h
--- a/Engine/Plugins/Experimental/ChaosCaching/Source/ChaosCachingEditor/Public/Chaos/ChaosCachingEditorPlugin.h
+++ b/Engine/Plugins/Experimental/ChaosCaching/Source/ChaosCachingEditor/Public/Chaos/ChaosCachingEditorPlugin.h
@@ -47,8 +47,6 @@ private:
 	void RegisterCachingSubMenu(UToolMenu* InMenu, FToolMenuSection* InSection);
 	
 	void OnCreateCacheManager();
-	void OnSetAllPlay();
-	void OnSetAllRecord();
 
 	FAssetTypeActions_ChaosCacheCollection* AssetTypeActions_ChaosCacheCollection;
 	FDelegateHandle StartupHandle;
diff --git a/Engine/Source/Runtime/Experimental/Chaos/Private/PhysicsProxy/GeometryCollectionPhysicsProxy.cpp b/Engine/Source/Runtime/Experimental/Chaos/Private/PhysicsProxy/GeometryCollectionPhysicsProxy.cpp
--- a/Engine/Source/Runtime/Experimental/Chaos/Private/PhysicsProxy/GeometryCollectionPhysicsProxy.cpp
+++ b/Engine/Source/Runtime/Experimental/Chaos/Private/PhysicsProxy/GeometryCollectionPhysicsProxy.cpp
@@ -466,8 +466,14 @@ void FGeometryCollectionPhysicsProxy::InitializeDynamicCollection(FGeometryDynam
 	//
 
 	TMap<FName, TSet<FName>> SkipList;
-	TSet<FName>& TransformGroupSkipList = SkipList.Emplace(FTransformCollection::TransformGroup);
-	TransformGroupSkipList.Add(DynamicCollection.SimplicialsAttribute);
+	TSet<FName>& KeepFromDynamicCollection = SkipList.Emplace(FTransformCollection::TransformGroup);
+	KeepFromDynamicCollection.Add(FTransformCollection::TransformAttribute);
+	KeepFromDynamicCollection.Add(FTransformCollection::ParentAttribute);
+	KeepFromDynamicCollection.Add(FTransformCollection::ChildrenAttribute);
+	KeepFromDynamicCollection.Add(FGeometryCollection::SimulationTypeAttribute);
+	KeepFromDynamicCollection.Add(DynamicCollection.SimplicialsAttribute);
+	KeepFromDynamicCollection.Add(DynamicCollection.ActiveAttribute);
+	KeepFromDynamicCollection.Add(DynamicCollection.CollisionGroupAttribute);
 	DynamicCollection.CopyMatchingAttributesFrom(RestCollection, &SkipList);
 
 
@@ -511,25 +517,11 @@ void FGeometryCollectionPhysicsProxy::InitializeDynamicCollection(FGeometryDynam
 	}
 
 	// Process Activity
+	const int32 NumTransforms = DynamicCollection.SimulatableParticles.Num();
+	for (int32 TransformIdx = 0; TransformIdx < NumTransforms; TransformIdx++)
 	{
-		const int32 NumTransforms = DynamicCollection.SimulatableParticles.Num();
-		if (!RestCollection.HasAttribute(FGeometryCollection::SimulatableParticlesAttribute, FTransformCollection::TransformGroup))
-		{
-			// If no simulation data is available then default to the simulation of just the rigid geometry.
-			for (int32 TransformIdx = 0; TransformIdx < NumTransforms; TransformIdx++)
-			{
-				if (DynamicCollection.Children[TransformIdx].Num())
-				{
-					DynamicCollection.SimulatableParticles[TransformIdx] = false;
-				}
-				else
-				{
-					DynamicCollection.SimulatableParticles[TransformIdx] = DynamicCollection.Active[TransformIdx];
-				}
-			}
-		}
+		DynamicCollection.SimulatableParticles[TransformIdx] = DynamicCollection.Active[TransformIdx];
 	}
-
 }
 
 int32 ReportTooManyChildrenNum = -1;
@@ -544,8 +536,6 @@ void FGeometryCollectionPhysicsProxy::InitializeBodiesPT(Chaos::FPBDRigidsSolver
 	{
 		const TManagedArray<int32>& TransformIndex = RestCollection->TransformIndex;
 		const TManagedArray<int32>& BoneMap = RestCollection->BoneMap;
-		const TManagedArray<int32>& Parent = RestCollection->Parent;
-		const TManagedArray<TSet<int32>>& Children = RestCollection->Children;
 		const TManagedArray<int32>& SimulationType = RestCollection->SimulationType;
 		const TManagedArray<FVector3f>& Vertex = RestCollection->Vertex;
 		const TManagedArray<float>& Mass = RestCollection->GetAttribute<float>("Mass", FTransformCollection::TransformGroup);
@@ -560,9 +550,11 @@ void FGeometryCollectionPhysicsProxy::InitializeBodiesPT(Chaos::FPBDRigidsSolver
 		const TManagedArray<FVector3f>& InitialLinearVelocity = DynamicCollection.InitialLinearVelocity;
 		const TManagedArray<FGeometryDynamicCollection::FSharedImplicit>& Implicits = DynamicCollection.Implicits;
 		const TManagedArray<TUniquePtr<FCollisionStructureManager::FSimplicial>>& Simplicials = DynamicCollection.Simplicials;
+		const TManagedArray<TSet<int32>>& Children = DynamicCollection.Children;
+		const TManagedArray<int32>& Parent = DynamicCollection.Parent;
 
 		TArray<FTransform> Transform;
-		GeometryCollectionAlgo::GlobalMatrices(DynamicCollection.Transform, DynamicCollection.Parent, Transform);
+		GeometryCollectionAlgo::GlobalMatrices(DynamicCollection.Transform, Parent, Transform);
 
 		//const int NumRigids = 0; // ryan - Since we're doing SOA, we start at zero?
 		int NumRigids = 0;
diff --git a/Engine/Source/Runtime/Experimental/GeometryCollectionEngine/Private/GeometryCollection/GeometryCollectionComponent.cpp b/Engine/Source/Runtime/Experimental/GeometryCollectionEngine/Private/GeometryCollection/GeometryCollectionComponent.cpp
--- a/Engine/Source/Runtime/Experimental/GeometryCollectionEngine/Private/GeometryCollection/GeometryCollectionComponent.cpp
+++ b/Engine/Source/Runtime/Experimental/GeometryCollectionEngine/Private/GeometryCollection/GeometryCollectionComponent.cpp
@@ -858,6 +858,58 @@ void UGeometryCollectionComponent::RefreshEmbeddedGeometry()
 	}
 }
 
+void UGeometryCollectionComponent::SetRestState(TArray<FTransform>&& InRestTransforms)
+{
+	// This method should never be called on a simulating component.
+	if (!PhysicsProxy)
+	{ 
+		RestTransforms = InRestTransforms;
+	
+		CalculateGlobalMatrices();
+
+		if (DynamicCollection)
+		{
+			SetInitialTransforms(RestTransforms);
+		}
+
+		FGeometryCollectionDynamicData* DynamicData = GDynamicDataPool.Allocate();
+		DynamicData->SetAllTransforms(GlobalMatrices);
+		DynamicData->IsDynamic = true;
+
+		if (SceneProxy)
+		{
+			if (SceneProxy->IsNaniteMesh())
+			{
+
+	#if WITH_EDITOR
+				// We need to do this in case we're controlled by Sequencer in editor, which doesn't invoke PostEditChangeProperty
+				SendRenderTransform_Concurrent();
+	#endif
+
+				FNaniteGeometryCollectionSceneProxy* GeometryCollectionSceneProxy = static_cast<FNaniteGeometryCollectionSceneProxy*>(SceneProxy);
+				ENQUEUE_RENDER_COMMAND(SendRenderDynamicData)(
+					[GeometryCollectionSceneProxy, DynamicData](FRHICommandListImmediate& RHICmdList)
+					{
+						GeometryCollectionSceneProxy->SetDynamicData_RenderThread(DynamicData);
+					}
+				);
+			}
+			else
+			{
+				FGeometryCollectionSceneProxy* GeometryCollectionSceneProxy = static_cast<FGeometryCollectionSceneProxy*>(SceneProxy);
+				ENQUEUE_RENDER_COMMAND(SendRenderDynamicData)(
+					[GeometryCollectionSceneProxy, DynamicData](FRHICommandListImmediate& RHICmdList)
+					{
+						GeometryCollectionSceneProxy->SetDynamicData_RenderThread(DynamicData);
+					}
+				);
+			}
+		}
+
+		RefreshEmbeddedGeometry();
+	}
+}
+
 void UGeometryCollectionComponent::InitializeComponent()
 {
 	Super::InitializeComponent();
@@ -1152,6 +1204,53 @@ void UGeometryCollectionComponent::UpdateRepData()
 	}
 }
 
+void UGeometryCollectionComponent::SetDynamicState(const Chaos::EObjectStateType& NewDynamicState)
+{
+	if (DynamicCollection)
+	{
+		TManagedArray<int32>& DynamicState = DynamicCollection->DynamicState;
+		for (int i = 0; i < DynamicState.Num(); i++)
+		{
+			DynamicState[i] = static_cast<int32>(NewDynamicState);
+		}
+	}
+}
+
+void UGeometryCollectionComponent::SetInitialTransforms(const TArray<FTransform>& InitialTransforms)
+{
+	if (DynamicCollection)
+	{
+		TManagedArray<FTransform>& Transform = DynamicCollection->Transform;
+		int32 MaxIdx = FMath::Min(Transform.Num(), InitialTransforms.Num());
+		for (int32 Idx = 0; Idx < MaxIdx; ++Idx)
+		{
+			Transform[Idx] = InitialTransforms[Idx];
+		}
+	}
+}
+
+void UGeometryCollectionComponent::SetInitialClusterBreaks(const TArray<int32>& ReleaseIndices)
+{
+	if (DynamicCollection)
+	{
+		TManagedArray<int32>& Parent = DynamicCollection->Parent;
+		TManagedArray <TSet<int32>>& Children = DynamicCollection->Children;
+		const int32 NumTransforms = Parent.Num();
+
+		for (int32 ReleaseIndex : ReleaseIndices)
+		{
+			if (ReleaseIndex < NumTransforms)
+			{
+				if (Parent[ReleaseIndex] > INDEX_NONE)
+				{
+					Children[Parent[ReleaseIndex]].Remove(ReleaseIndex);
+					Parent[ReleaseIndex] = INDEX_NONE;
+				}
+			}
+		}
+	}
+}
+
 void SetHierarchyStrain(Chaos::TPBDRigidClusteredParticleHandle<Chaos::FReal, 3>* P, TMap<Chaos::TPBDRigidClusteredParticleHandle<Chaos::FReal, 3>*, TArray<Chaos::TPBDRigidParticleHandle<Chaos::FReal, 3>*>>& Map, float Strain)
 {
 	TArray<Chaos::TPBDRigidParticleHandle<Chaos::FReal, 3>*>* Children = Map.Find(P);
@@ -1299,6 +1398,8 @@ void UGeometryCollectionComponent::InitConstantData(FGeometryCollectionConstantD
 
 			ConstantData->Sections[SectionIndex] = MoveTemp(Section);
 		}
+		
+		
 		ConstantData->NumTransforms = Collection->NumElements(FGeometryCollection::TransformGroup);
 		ConstantData->LocalBounds = LocalBounds;
 
@@ -1347,183 +1448,7 @@ FGeometryCollectionDynamicData* UGeometryCollectionComponent::InitDynamicData(bo
 	const bool bEditorMode = bShowBoneColors || bEnableBoneSelection;
 	const bool bIsDynamic  = GetIsObjectDynamic() || bEditorMode || bInitialization;
 
-	if (CachePlayback && CacheParameters.TargetCache)
-	{
-		// If we are already on the current cached frame, return
-		if (FMath::IsNearlyEqual(CurrentCacheTime, DesiredCacheTime) && GlobalMatrices.Num() != 0)
-		{
-			return DynamicData;
-		}
-
-		// #todo(dmp): find a better place to calculate and store this
-		float CacheDt = CacheParameters.TargetCache->GetData()->GetDt();
-
-		// If we want the state before the first cached frame, then the collection doesn't need to be dynamic since it'll render the pre-fractured geometry
-		const bool bIsCacheDynamic = DesiredCacheTime > CacheDt;
-
-		// Bad case here.  Sequencer starts at non zero position  We cannot correctly reconstruct the current frame, so give a warning
-		const bool bOutOfSync = GlobalMatrices.Num() == 0;
-
-		// If the input simulation time to playback is the first frame, reset simulation time.
-		if (bOutOfSync || DesiredCacheTime <= CacheDt ||  FMath::IsNearlyEqual(CurrentCacheTime, FLT_MAX))
-		{
-			GeometryCollectionAlgo::GlobalMatrices(GetTransformArray(), GetParentArray(), GlobalMatrices);
-
-			DynamicData = GDynamicDataPool.Allocate();
-			DynamicData->SetAllTransforms(GlobalMatrices);
-			DynamicData->IsDynamic = bIsCacheDynamic;
-			DynamicData->IsLoading = GetIsObjectLoading();
-
-			EventsPlayed.Reset();
-			EventsPlayed.AddDefaulted(CacheParameters.TargetCache->GetData()->Records.Num());
-
-			if (bOutOfSync)
-			{
-				UE_LOG(UGCC_LOG, Warning, TEXT("Cache out of sync - must rewind sequencer to start frame"));
-			}
-			else
-			{
-				CurrentCacheTime = DesiredCacheTime;
-			}
-		}
-		else if (DesiredCacheTime >= CurrentCacheTime)
-		{
-			const TManagedArray<FTransform>& Transform = RestCollection->GetGeometryCollection()->Transform;
-			const TManagedArray<TSet<int32>>& Children = GetChildrenArray();
-			const TManagedArray<FTransform>& MassToLocal = RestCollection->GetGeometryCollection()->GetAttribute<FTransform>("MassToLocal", FGeometryCollection::TransformGroup);
-
-			int32 NumSteps = floor((DesiredCacheTime - CurrentCacheTime) / CacheDt);
-			float LastDt = FMath::Fmod(DesiredCacheTime - CurrentCacheTime, CacheDt);
-			NumSteps += LastDt > SMALL_NUMBER ? 1 : 0;
-			
-			FTransform ActorToWorld = GetComponentTransform();
-
-			DynamicData = GDynamicDataPool.Allocate();
-			DynamicData->ChangedCount = 0;
-			DynamicData->IsDynamic = bIsCacheDynamic;
-			DynamicData->IsLoading = GetIsObjectLoading();
-
-			bool HasAnyActiveTransforms = false;
-
-			// Jump ahead in increments of CacheDt evaluating the cache until we reach our desired time
-			for (int32 Step = 0; Step < NumSteps; ++Step)
-			{
-				float TimeIncrement = Step == NumSteps - 1 ? LastDt : CacheDt;
-				CurrentCacheTime += TimeIncrement;
-			
-				DynamicData->SetPrevTransforms(GlobalMatrices);
-
-				const FRecordedFrame* FirstFrame = nullptr;
-				const FRecordedFrame* SecondFrame = nullptr;
-				CacheParameters.TargetCache->GetData()->GetFramesForTime(CurrentCacheTime, FirstFrame, SecondFrame);
-
-				if (FirstFrame && !SecondFrame)
-				{
-					const TArray<FTransform>& FirstFrameTransforms = FirstFrame->Transforms;
-					const TArray<int32>& TransformIndices = FirstFrame->TransformIndices;
-
-					const int32 NumActives = TransformIndices.Num();
-					
-					if (NumActives > 0)
-					{
-						HasAnyActiveTransforms = true;
-					}
-
-					for (int32 TransformIndex = 0; TransformIndex < NumActives; ++TransformIndex)
-					{
-						const int32 InternalIndexTmp = TransformIndices[TransformIndex];
-						
-						if (InternalIndexTmp >= GlobalMatrices.Num())
-						{
-							UE_LOG
-							(
-								UGCC_LOG, Error,
-								TEXT("%s: TargetCache (%s) is out of sync with GeometryCollection.  Regenerate the cache."),
-								*RestCollection->GetName(), *CacheParameters.TargetCache->GetName()
-							);
-
-							DynamicData->SetAllTransforms(GlobalMatrices);
-							return DynamicData;
-						}
-
-						// calculate global matrix for current
-						FTransform ParticleToWorld = FirstFrameTransforms[TransformIndex];
-
-						FTransform CurrGlobalTransform = MassToLocal[InternalIndexTmp].GetRelativeTransformReverse(ParticleToWorld).GetRelativeTransform(ActorToWorld);
-						CurrGlobalTransform.NormalizeRotation();
-						GlobalMatrices[InternalIndexTmp] = CurrGlobalTransform.ToMatrixWithScale();
-						
-						// Traverse from active parent node down to all children and set global transforms
-						GeometryCollectionAlgo::GlobalMatricesFromRoot(InternalIndexTmp, Transform, Children, GlobalMatrices);
-					}
-				}
-				else if (FirstFrame && SecondFrame && CurrentCacheTime > FirstFrame->Timestamp)
-				{
-					const float Alpha = (CurrentCacheTime - FirstFrame->Timestamp) / (SecondFrame->Timestamp - FirstFrame->Timestamp);
-					checkSlow(0 <= Alpha && Alpha <= 1.0f);
-
-					const int32 NumActives = SecondFrame->TransformIndices.Num();
-
-					if (NumActives > 0)
-					{
-						HasAnyActiveTransforms = true;
-					}
-
-					for (int32 Index = 0; Index < NumActives; ++Index)
-					{
-						const int32 InternalIndexTmp = SecondFrame->TransformIndices[Index];
-						
-						// Check if transform index is valid
-						if (InternalIndexTmp >= GlobalMatrices.Num())
-						{
-							UE_LOG
-							(
-								UGCC_LOG, Error,
-								TEXT("%s: TargetCache (%s) is out of sync with GeometryCollection.  Regenerate the cache."), 
-								*RestCollection->GetName(), *CacheParameters.TargetCache->GetName()
-							);
-							DynamicData->SetAllTransforms(GlobalMatrices);
-							return DynamicData;
-						}
-
-						const int32 PreviousIndexSlot = Index < SecondFrame->PreviousTransformIndices.Num() ? SecondFrame->PreviousTransformIndices[Index] : INDEX_NONE;
-
-						if (PreviousIndexSlot != INDEX_NONE)
-						{
-							FTransform ParticleToWorld;
-							ParticleToWorld.Blend(FirstFrame->Transforms[PreviousIndexSlot], SecondFrame->Transforms[Index], Alpha);
-							
-							FTransform CurrGlobalTransform = MassToLocal[InternalIndexTmp].GetRelativeTransformReverse(ParticleToWorld).GetRelativeTransform(ActorToWorld);
-							CurrGlobalTransform.NormalizeRotation();
-							GlobalMatrices[InternalIndexTmp] = CurrGlobalTransform.ToMatrixWithScale();
-
-							// Traverse from active parent node down to all children and set global transforms
-							GeometryCollectionAlgo::GlobalMatricesFromRoot(InternalIndexTmp, Transform, Children, GlobalMatrices);
-						}
-						else
-						{
-							FTransform ParticleToWorld = SecondFrame->Transforms[Index];
-							
-							FTransform CurrGlobalTransform = MassToLocal[InternalIndexTmp].GetRelativeTransformReverse(ParticleToWorld).GetRelativeTransform(ActorToWorld);
-							CurrGlobalTransform.NormalizeRotation();
-							GlobalMatrices[InternalIndexTmp] = CurrGlobalTransform.ToMatrixWithScale();
-
-							// Traverse from active parent node down to all children and set global transforms
-							GeometryCollectionAlgo::GlobalMatricesFromRoot(InternalIndexTmp, Transform, Children, GlobalMatrices);
-						}
-					}
-				}
-				DynamicData->SetTransforms(GlobalMatrices);
-			}
-
-			// Check if transforms at start of this tick are the same as what is calculated from the cache
-			if (HasAnyActiveTransforms)
-			{
-				DynamicData->DetermineChanges();
-			}
-		}
-	}
-	else if (bIsDynamic)
+	if (bIsDynamic)
 	{
 		DynamicData = GDynamicDataPool.Allocate();
 		DynamicData->IsDynamic = true;
@@ -1737,6 +1662,11 @@ void UGeometryCollectionComponent::ResetDynamicCollection()
 		SetRenderStateDirty();
 	}
 
+	if (RestTransforms.Num() > 0)
+	{
+		SetInitialTransforms(RestTransforms);
+	}
+
 	if (RestCollection)
 	{
 		CalculateGlobalMatrices();
@@ -1808,13 +1738,24 @@ void UGeometryCollectionComponent::OnCreatePhysicsState()
 				}
 			}
 
-			TManagedArray<bool> & Active = DynamicCollection->Active;
+			TManagedArray<bool>& Active = DynamicCollection->Active;
+			if (RestCollection->GetGeometryCollection()->HasAttribute(FGeometryCollection::SimulatableParticlesAttribute, FTransformCollection::TransformGroup))
+			{
+				TManagedArray<bool>* SimulatableParticles = RestCollection->GetGeometryCollection()->FindAttribute<bool>(FGeometryCollection::SimulatableParticlesAttribute, FTransformCollection::TransformGroup);
+				for (int i = 0; i < Active.Num(); i++)
+				{
+					Active[i] = (*SimulatableParticles)[i];
+				}
+			}
+			else
 			{
+				// If no simulation data is available then default to the simulation of just the rigid geometry.
 				for (int i = 0; i < Active.Num(); i++)
 				{
-					Active[i] = true;
+					Active[i] = RestCollection->GetGeometryCollection()->IsRigid(i);
 				}
 			}
+			
 			TManagedArray<int32> & CollisionGroupArray = DynamicCollection->CollisionGroup;
 			{
 				for (int i = 0; i < CollisionGroupArray.Num(); i++)
@@ -2024,6 +1965,13 @@ void UGeometryCollectionComponent::SetRestCollection(const UGeometryCollection*
 	{
 		RestCollection = RestCollectionIn;
 
+		const int32 NumTransforms = RestCollection->GetGeometryCollection()->NumElements(FGeometryCollection::TransformGroup);
+		RestTransforms.SetNum(NumTransforms);
+		for (int32 Idx = 0; Idx < NumTransforms; ++Idx)
+		{
+			RestTransforms[Idx] = RestCollection->GetGeometryCollection()->Transform[Idx];
+		}
+
 		CalculateGlobalMatrices();
 		CalculateLocalBounds();
 
@@ -2641,11 +2589,24 @@ void UGeometryCollectionComponent::CalculateGlobalMatrices()
 		GlobalMatrices.Append(Results->GlobalTransforms);	
 	}
 	else
-	{
-		// Have to fully rebuild
-		GeometryCollectionAlgo::GlobalMatrices(GetTransformArray(), GetParentArray(), GlobalMatrices);
+	{	
+		// If hierarchy topology has changed, the RestTransforms is invalidated.
+		if (RestTransforms.Num() != GetTransformArray().Num())
+		{
+			RestTransforms.Empty();
+		}
+		
+		if (!DynamicCollection && RestTransforms.Num() > 0)
+		{
+			GeometryCollectionAlgo::GlobalMatrices(RestTransforms, GetParentArray(), GlobalMatrices);
+		}
+		else
+		{
+			// Have to fully rebuild
+			GeometryCollectionAlgo::GlobalMatrices(GetTransformArray(), GetParentArray(), GlobalMatrices);
+		}
 	}
-
+	
 #if WITH_EDITOR
 	if (GlobalMatrices.Num() > 0)
 	{
@@ -2779,7 +2740,7 @@ void UGeometryCollectionComponent::ClearEmbeddedGeometry()
 
 	for (UActorComponent* TargetComponent : TargetComponents)
 	{
-		if (TargetComponent->GetOuter() == this)
+		if ((TargetComponent->GetOuter() == this) || (TargetComponent->GetOuter()->IsPendingKill()))
 		{
 			if (UInstancedStaticMeshComponent* ISMComponent = Cast<UInstancedStaticMeshComponent>(TargetComponent))
 			{
@@ -2815,6 +2776,7 @@ void UGeometryCollectionComponent::InitializeEmbeddedGeometry()
 					ISMC->SetCastShadow(false);
 					ISMC->SetMobility(EComponentMobility::Stationary);
 					ISMC->SetupAttachment(this);
+					ActorOwner->AddInstanceComponent(ISMC);
 					ISMC->RegisterComponent();
 
 					EmbeddedGeometryComponents.Add(ISMC);
@@ -2824,7 +2786,6 @@ void UGeometryCollectionComponent::InitializeEmbeddedGeometry()
 
 		CalculateGlobalMatrices();
 		RefreshEmbeddedGeometry();
-
 	}
 }
 
diff --git a/Engine/Source/Runtime/Experimental/GeometryCollectionEngine/Public/GeometryCollection/GeometryCollectionComponent.h b/Engine/Source/Runtime/Experimental/GeometryCollectionEngine/Public/GeometryCollection/GeometryCollectionComponent.h
--- a/Engine/Source/Runtime/Experimental/GeometryCollectionEngine/Public/GeometryCollection/GeometryCollectionComponent.h
+++ b/Engine/Source/Runtime/Experimental/GeometryCollectionEngine/Public/GeometryCollection/GeometryCollectionComponent.h
@@ -508,6 +508,10 @@ public:
 	UPROPERTY()
 	FGeomComponentCacheParameters CacheParameters;
 
+	/** Optional transforms to initialize scene proxy if difference from the RestCollection. */
+	UPROPERTY()
+	TArray<FTransform> RestTransforms;
+
 	/**
 	*  SetDynamicState
 	*    This function will dispatch a command to the physics thread to apply
@@ -638,6 +642,18 @@ public:
 	/** Update instanced static mesh components to reflect internal embedded geometry state. */
 	void RefreshEmbeddedGeometry();
 
+	// #todo should this only be available in editor?
+	void SetRestState(TArray<FTransform>&& InRestTransforms);
+
+	/** Set the dynamic state for all bodies in the DynamicCollection. */
+	void SetDynamicState(const Chaos::EObjectStateType& NewDynamicState);
+
+	/** Set transforms for all bodies in the DynamicCollection. */
+	void SetInitialTransforms(const TArray<FTransform>& InitialTransforms);
+
+	/** Modify DynamicCollection transform hierarchy to effect cluster breaks releasing the specified indices. */
+	void SetInitialClusterBreaks(const TArray<int32>& ReleaseIndices);
+
 public:
 	UPROPERTY(BlueprintAssignable, Category = "Collision")
 	FOnChaosPhysicsCollision OnChaosPhysicsCollision;
@@ -647,17 +663,20 @@ public:
 
 	// IChaosNotifyHandlerInterface
 	virtual void DispatchChaosPhysicsCollisionBlueprintEvents(const FChaosPhysicsCollisionInfo& CollisionInfo) override;
-
-protected:
 	
-	/** Call SetNotifyBreaks to set this at runtime. */
+	/** If true, this component will generate breaking events that other systems may subscribe to. */
 	UPROPERTY(EditAnywhere, BlueprintReadOnly, Category = "ChaosPhysics|General")
 	bool bNotifyBreaks;
 
-	/** If true, this component will get Chaos-specific collision notification events (@see IChaosNotifyHandlerInterface) */
+	/** If true, this component will generate collision events that other systems may subscribe to. */
 	UPROPERTY(EditAnywhere, BlueprintReadOnly, Category = "ChaosPhysics|General")
 	bool bNotifyCollisions;
 
+	/** If true, this component will generate trailing events that other systems may subscribe to. */
+	UPROPERTY(EditAnywhere, BlueprintReadOnly, Category = "ChaosPhysics|General")
+	bool bNotifyTrailing;
+
+protected:
 	/** Display Bone Colors instead of assigned materials */
 	UPROPERTY(EditAnywhere, BlueprintReadWrite, Category = "ChaosPhysics|General")
 	bool bShowBoneColors;
