ChaosCaching
-- corrected editor only methods

diff --git a/Engine/Plugins/Experimental/ChaosCaching/Source/ChaosCaching/Private/Chaos/CacheManagerActor.cpp b/Engine/Plugins/Experimental/ChaosCaching/Source/ChaosCaching/Private/Chaos/CacheManagerActor.cpp
--- a/Engine/Plugins/Experimental/ChaosCaching/Source/ChaosCaching/Private/Chaos/CacheManagerActor.cpp
+++ b/Engine/Plugins/Experimental/ChaosCaching/Source/ChaosCaching/Private/Chaos/CacheManagerActor.cpp
@@ -104,6 +104,7 @@ void AChaosCacheManager::TickActor(float DeltaTime, enum ELevelTick TickType, FA
 	}
 }
 
+#if WITH_EDITOR
 void AChaosCacheManager::PostEditChangeProperty(FPropertyChangedEvent& PropertyChangedEvent)
 {
 	Super::PostEditChangeProperty(PropertyChangedEvent);
@@ -118,9 +119,7 @@ void AChaosCacheManager::PostEditChangeProperty(FPropertyChangedEvent& PropertyC
 			CacheMode = ECacheMode::None;
 		}
 
-#if WITH_EDITOR
 		SetObservedComponentProperties(CacheMode);
-#endif
 		
 		if (CacheMode != ECacheMode::Record)
 		{
@@ -139,6 +138,7 @@ void AChaosCacheManager::PostEditChangeProperty(FPropertyChangedEvent& PropertyC
 		}
 	}
 }
+#endif
 
 void AChaosCacheManager::SetStartTime(float InStartTime)
 {
@@ -613,7 +613,12 @@ FObservedComponent& AChaosCacheManager::AddNewObservedComponent(UPrimitiveCompon
 	NewEntry.ComponentRef.PathToComponent = InComponent->GetPathName(InComponent->GetOwner());
 	NewEntry.ComponentRef.OtherActor      = InComponent->GetOwner();
 
-	NewEntry.CacheName                    = MakeUniqueObjectName(CacheCollection, UChaosCache::StaticClass(), FName(InComponent->GetOwner()->GetActorLabel()));
+	FName CacheName = NAME_None;
+#if WITH_EDITOR
+	CacheName = FName(InComponent->GetOwner()->GetActorLabel());
+#endif
+
+	NewEntry.CacheName = MakeUniqueObjectName(CacheCollection, UChaosCache::StaticClass(), CacheName);
 
 	return NewEntry;
 }
diff --git a/Engine/Plugins/Experimental/ChaosCaching/Source/ChaosCaching/Public/Chaos/CacheManagerActor.h b/Engine/Plugins/Experimental/ChaosCaching/Source/ChaosCaching/Public/Chaos/CacheManagerActor.h
--- a/Engine/Plugins/Experimental/ChaosCaching/Source/ChaosCaching/Public/Chaos/CacheManagerActor.h
+++ b/Engine/Plugins/Experimental/ChaosCaching/Source/ChaosCaching/Public/Chaos/CacheManagerActor.h
@@ -143,7 +143,9 @@ public:
 	friend class IChaosCachingEditorPlugin;
 #endif
 
-	void PostEditChangeProperty(FPropertyChangedEvent& PropertyChangedEvent) override;
+#if WITH_EDITOR
+	virtual void PostEditChangeProperty(FPropertyChangedEvent& PropertyChangedEvent) override;
+#endif
 	/** end UObject interface */
 
 	/** Expose StartTime property to Sequencer. GetStartTime will be called on keys. */
