Chaos Caching
-- Removed GeometryCollectionComponent dereferencing from PhysicsThread code
-- Uses PhysicsProxy simulation parameters to determine if events are recorded/played

diff --git a/Engine/Plugins/Experimental/ChaosCaching/Source/ChaosCaching/Private/Chaos/Adapters/GeometryCollectionComponentCacheAdapter.cpp b/Engine/Plugins/Experimental/ChaosCaching/Source/ChaosCaching/Private/Chaos/Adapters/GeometryCollectionComponentCacheAdapter.cpp
--- a/Engine/Plugins/Experimental/ChaosCaching/Source/ChaosCaching/Private/Chaos/Adapters/GeometryCollectionComponentCacheAdapter.cpp
+++ b/Engine/Plugins/Experimental/ChaosCaching/Source/ChaosCaching/Private/Chaos/Adapters/GeometryCollectionComponentCacheAdapter.cpp
@@ -52,7 +52,7 @@ namespace Chaos
 
 		UGeometryCollectionComponent*    Comp  = CastChecked<UGeometryCollectionComponent>(InComp);
 		FGeometryCollectionPhysicsProxy* Proxy = Comp->GetPhysicsProxy();
-		const FTransform WorldToComponent = Comp->GetComponentTransform().Inverse();
+		const FTransform WorldToComponent = Proxy->GetSimParameters().WorldTransform.Inverse();
 
 		FPhysScene_Chaos& Scene = *(Comp->GetWorld()->GetPhysicsScene());
 		
@@ -79,8 +79,6 @@ namespace Chaos
 		FGeometryDynamicCollection&            Collection       = Proxy->GetPhysicsCollection();
 		const TManagedArray<FTransform>&       MassToLocal		= Collection.MassToLocal;
 		const TManagedArray<int32>&            Parents			= RestCollection->Parent;
-		const TManagedArray<int32>&            TransformIndices = RestCollection->TransformIndex;
-		const TManagedArray<FTransform>&       Transforms       = Collection.Transform;
 		const TArray<FBreakingData>&		   Breaks           = Solver->GetEvolution()->GetRigidClustering().GetAllClusterBreakings();
 		const FTransform					   WorldToActor		= InRootTransform.Inverse();
 
@@ -172,7 +170,7 @@ namespace Chaos
 			}
 		}
 
-		if (bHandleBreaks && BreakingDataArray && ProxyCachedEventData.ProxyBreakingDataIndices)
+		if (Proxy->GetSimParameters().bGenerateBreakingData && BreakingDataArray && ProxyCachedEventData.ProxyBreakingDataIndices)
 		{
 			for (int32 Index : *ProxyCachedEventData.ProxyBreakingDataIndices)
 			{
@@ -188,7 +186,7 @@ namespace Chaos
 			}
 		}
 		
-		if (bHandleCollisions && CollisionDataArray && ProxyCachedEventData.ProxyCollisionDataIndices)
+		if (Proxy->GetSimParameters().bGenerateCollisionData && CollisionDataArray && ProxyCachedEventData.ProxyCollisionDataIndices)
 		{
 			for (int32 Index : *ProxyCachedEventData.ProxyCollisionDataIndices)
 			{
@@ -201,7 +199,7 @@ namespace Chaos
 			}
 		}
 		
-		if (bHandleTrailing && TrailingDataArray && ProxyCachedEventData.ProxyTrailingDataIndices)
+		if (Proxy->GetSimParameters().bGenerateTrailingData && TrailingDataArray && ProxyCachedEventData.ProxyTrailingDataIndices)
 		{
 			for (int32 Index : *ProxyCachedEventData.ProxyTrailingDataIndices)
 			{
@@ -234,7 +232,7 @@ namespace Chaos
 
 		UGeometryCollectionComponent*    Comp  = CastChecked<UGeometryCollectionComponent>(InComponent);
 		FGeometryCollectionPhysicsProxy* Proxy = Comp->GetPhysicsProxy();
-		const FTransform ComponentToWorld = Comp->GetComponentTransform();
+		const FTransform ComponentToWorld = Proxy->GetSimParameters().WorldTransform;
 
 		if(!Proxy)
 		{
@@ -251,8 +249,6 @@ namespace Chaos
 
 		FGeometryDynamicCollection&      Collection       = Proxy->GetPhysicsCollection();
 		const TManagedArray<FTransform>& MassToLocal	  = Collection.MassToLocal;
-		const TManagedArray<int32>&      TransformIndices = RestCollection->TransformIndex;
-		const TManagedArray<FTransform>& Transforms       = Collection.Transform;
 		TArray<FClusterParticle*>        Particles        = Proxy->GetParticles();
 
 		FCacheEvaluationContext Context(TickRecord);
@@ -337,7 +333,7 @@ namespace Chaos
 			}
 		}
 
-		if (bHandleBreaks && BreakingEvents)
+		if (Proxy->GetSimParameters().bGenerateBreakingData && BreakingEvents)
 		{
 			const FSolverBreakingEventFilter* SolverBreakingEventFilter = Solver->GetEventFilters()->GetBreakingFilter();
 
@@ -385,7 +381,7 @@ namespace Chaos
 			}
 		}
 
-		if (bHandleTrailing && TrailingEvents)
+		if (Proxy->GetSimParameters().bGenerateTrailingData && TrailingEvents)
 		{
 			const FSolverTrailingEventFilter* SolverTrailingEventFilter = Solver->GetEventFilters()->GetTrailingFilter();
 
@@ -432,7 +428,7 @@ namespace Chaos
 			}
 		}
 
-		if (bHandleCollisions && CollisionEvents)
+		if (Proxy->GetSimParameters().bGenerateCollisionData && CollisionEvents)
 		{
 
 			const FSolverCollisionEventFilter* SolverCollisionEventFilter = Solver->GetEventFilters()->GetCollisionFilter();
@@ -673,6 +669,8 @@ namespace Chaos
 		UGeometryCollectionComponent*    Comp     = CastChecked<UGeometryCollectionComponent>(InComponent);
 		FGeometryCollectionPhysicsProxy* Proxy    = Comp->GetPhysicsProxy();
 
+		const FSimulationParameters& SimulationParameters = Proxy->GetSimParameters();
+
 		if (!Proxy)
 		{
 			return false;
@@ -685,10 +683,6 @@ namespace Chaos
 			return false;
 		}
 
-		bHandleBreaks = Comp->bNotifyBreaks;
-		bHandleCollisions = Comp->bNotifyCollisions;
-		bHandleTrailing = Comp->bNotifyTrailing;
-		
 		// We need secondary event data to record event information into the cache
 		Solver->SetGenerateBreakingData(true);
 		Solver->SetGenerateCollisionData(true);
@@ -700,17 +694,17 @@ namespace Chaos
 			Chaos::FEventManager* EventManager = Solver->GetEventManager();
 			if (EventManager)
 			{
-				if (bHandleBreaks)
+				if (SimulationParameters.bGenerateBreakingData)
 				{ 
 					EventManager->RegisterHandler<Chaos::FBreakingEventData>(Chaos::EEventType::Breaking, const_cast<FGeometryCollectionCacheAdapter*>(this), &FGeometryCollectionCacheAdapter::HandleBreakingEvents);
 				}
 
-				if (bHandleCollisions)
+				if (SimulationParameters.bGenerateCollisionData)
 				{ 
 					EventManager->RegisterHandler<Chaos::FCollisionEventData>(Chaos::EEventType::Collision, const_cast<FGeometryCollectionCacheAdapter*>(this), &FGeometryCollectionCacheAdapter::HandleCollisionEvents);
 				}
 
-				if (bHandleTrailing)
+				if (SimulationParameters.bGenerateTrailingData)
 				{ 
 					EventManager->RegisterHandler<Chaos::FTrailingEventData>(Chaos::EEventType::Trailing, const_cast<FGeometryCollectionCacheAdapter*>(this), &FGeometryCollectionCacheAdapter::HandleTrailingEvents);
 				}
@@ -732,10 +726,6 @@ namespace Chaos
 
 		if (UGeometryCollectionComponent* Comp = CastChecked<UGeometryCollectionComponent>(InComponent))
 		{
-			bHandleBreaks = Comp->bNotifyBreaks;
-			bHandleCollisions = Comp->bNotifyCollisions;
-			bHandleTrailing = Comp->bNotifyTrailing;
-			
 			if (const UGeometryCollection* RestCollection = Comp->GetRestCollection())
 			{
 				// Set up initial conditions... 
diff --git a/Engine/Plugins/Experimental/ChaosCaching/Source/ChaosCaching/Public/Chaos/Adapters/GeometryCollectionComponentCacheAdapter.h b/Engine/Plugins/Experimental/ChaosCaching/Source/ChaosCaching/Public/Chaos/Adapters/GeometryCollectionComponentCacheAdapter.h
--- a/Engine/Plugins/Experimental/ChaosCaching/Source/ChaosCaching/Public/Chaos/Adapters/GeometryCollectionComponentCacheAdapter.h
+++ b/Engine/Plugins/Experimental/ChaosCaching/Source/ChaosCaching/Public/Chaos/Adapters/GeometryCollectionComponentCacheAdapter.h
@@ -261,10 +261,7 @@ namespace Chaos
 	private:
 		TArray<int32> GatherAllBreaksUpToTime(UChaosCache* InCache, float InTime) const;
 
-		bool bHandleBreaks;
-		bool bHandleCollisions;
-		bool bHandleTrailing;
-
+		
 		TMap<IPhysicsProxyBase*, FCachedEventData> CachedData;
 		const FBreakingDataArray* BreakingDataArray;
 		const FCollisionDataArray* CollisionDataArray;
diff --git a/Engine/Source/Runtime/Experimental/Chaos/Public/GeometryCollection/GeometryCollectionSimulationCoreTypes.h b/Engine/Source/Runtime/Experimental/Chaos/Public/GeometryCollection/GeometryCollectionSimulationCoreTypes.h
--- a/Engine/Source/Runtime/Experimental/Chaos/Public/GeometryCollection/GeometryCollectionSimulationCoreTypes.h
+++ b/Engine/Source/Runtime/Experimental/Chaos/Public/GeometryCollection/GeometryCollectionSimulationCoreTypes.h
@@ -158,105 +158,6 @@ struct FSharedSimulationParameters
 	int32 MaximumCollisionParticleCount;
 };
 
-struct FCollisionDataSimulationParameters
-{
-	FCollisionDataSimulationParameters()
-		: DoGenerateCollisionData(false)
-		, SaveCollisionData(false)
-
-		, CollisionDataSizeMax(512)
-		, DoCollisionDataSpatialHash(false)
-		, CollisionDataSpatialHashRadius(50.f)
-		, MaxCollisionPerCell(1)
-	{}
-
-	FCollisionDataSimulationParameters(bool InDoGenerateCollisionData
-		, bool InSaveCollisionData
-		, int32 InCollisionDataSizeMax
-		, bool InDoCollisionDataSpatialHash
-		, float InCollisionDataSpatialHashRadius
-		, int32 InMaxCollisionPerCell)
-		: DoGenerateCollisionData(InDoGenerateCollisionData)
-		, SaveCollisionData(InSaveCollisionData)
-		, CollisionDataSizeMax(InCollisionDataSizeMax)
-		, DoCollisionDataSpatialHash(InDoCollisionDataSpatialHash)
-		, CollisionDataSpatialHashRadius(InCollisionDataSpatialHashRadius)
-		, MaxCollisionPerCell(InMaxCollisionPerCell)
-	{}
-
-	bool DoGenerateCollisionData;
-	bool SaveCollisionData;
-	int32 CollisionDataSizeMax;
-	bool DoCollisionDataSpatialHash;
-	float CollisionDataSpatialHashRadius;
-	int32 MaxCollisionPerCell;
-
-	FCollisionFilterData QueryData;
-	FCollisionFilterData SimData;
-};
-
-struct FBreakingDataSimulationParameters
-{
-	FBreakingDataSimulationParameters()
-		: DoGenerateBreakingData(false)
-		, SaveBreakingData(false)
-		, BreakingDataSizeMax(512)
-		, DoBreakingDataSpatialHash(false)
-		, BreakingDataSpatialHashRadius(15.f)
-		, MaxBreakingPerCell(1)
-	{}
-
-	FBreakingDataSimulationParameters(bool InDoGenerateBreakingData
-		, bool InSaveBreakingData
-		, int32 InBreakingDataSizeMax
-		, bool InDoBreakingDataSpatialHash
-		, float InBreakingDataSpatialHashRadius
-		, int32 InMaxBreakingPerCell)
-		: DoGenerateBreakingData(InDoGenerateBreakingData)
-		, SaveBreakingData(InSaveBreakingData)
-		, BreakingDataSizeMax(InBreakingDataSizeMax)
-		, DoBreakingDataSpatialHash(InDoBreakingDataSpatialHash)
-		, BreakingDataSpatialHashRadius(InBreakingDataSpatialHashRadius)
-		, MaxBreakingPerCell(InMaxBreakingPerCell)
-	{}
-
-	bool DoGenerateBreakingData;
-	bool SaveBreakingData;
-	int32 BreakingDataSizeMax;
-	bool DoBreakingDataSpatialHash;
-	float BreakingDataSpatialHashRadius;
-	int32 MaxBreakingPerCell;
-};
-
-struct FTrailingDataSimulationParameters
-{
-	FTrailingDataSimulationParameters()
-		: DoGenerateTrailingData(false)
-		, SaveTrailingData(false)
-		, TrailingDataSizeMax(512)
-		, TrailingMinSpeedThreshold(200.f)
-		, TrailingMinVolumeThreshold(10000.f)
-	{}
-
-	FTrailingDataSimulationParameters(bool InDoGenerateTrailingData
-		, bool InSaveTrailingData
-		, int32 InTrailingDataSizeMax
-		, float InTrailingMinSpeedThreshold
-		, float InTrailingMinVolumeThreshold)
-		: DoGenerateTrailingData(InDoGenerateTrailingData)
-		, SaveTrailingData(InSaveTrailingData)
-		, TrailingDataSizeMax(InTrailingDataSizeMax)
-		, TrailingMinSpeedThreshold(InTrailingMinSpeedThreshold)
-		, TrailingMinVolumeThreshold(InTrailingMinVolumeThreshold)
-	{}
-
-	bool DoGenerateTrailingData;
-	bool SaveTrailingData;
-	int32 TrailingDataSizeMax;
-	float TrailingMinSpeedThreshold;
-	float TrailingMinVolumeThreshold;
-};
-
 struct FSimulationParameters
 {
 	FSimulationParameters()
@@ -281,6 +182,9 @@ struct FSimulationParameters
 		, CacheBeginTime(0.0f)
 		, ReverseCacheBeginTime(0.0f)
 		, bClearCache(false)
+		, bGenerateBreakingData(false)
+		, bGenerateCollisionData(false)
+		, bGenerateTrailingData(false)
 		, RemoveOnFractureEnabled(false)
 		, SimulationFilterData()
 		, QueryFilterData()
@@ -311,9 +215,9 @@ struct FSimulationParameters
 		, ReverseCacheBeginTime(Other.ReverseCacheBeginTime)
 		, bClearCache(Other.bClearCache)
 		, PhysicalMaterialHandle(Other.PhysicalMaterialHandle)
-		, CollisionData(Other.CollisionData)
-		, BreakingData(Other.BreakingData)
-		, TrailingData(Other.TrailingData)
+		, bGenerateBreakingData(Other.bGenerateBreakingData)
+		, bGenerateCollisionData(Other.bGenerateCollisionData)
+		, bGenerateTrailingData(Other.bGenerateTrailingData)
 		, Shared(Other.Shared)
 		, RemoveOnFractureEnabled(false)
 		, SimulationFilterData(Other.SimulationFilterData)
@@ -366,9 +270,9 @@ struct FSimulationParameters
 
 	Chaos::FMaterialHandle PhysicalMaterialHandle;
 
-	FCollisionDataSimulationParameters CollisionData;
-	FBreakingDataSimulationParameters BreakingData;
-	FTrailingDataSimulationParameters TrailingData;
+	bool bGenerateBreakingData;
+	bool bGenerateCollisionData;
+	bool bGenerateTrailingData;
 
 	FSharedSimulationParameters Shared;
 
diff --git a/Engine/Source/Runtime/Experimental/GeometryCollectionEngine/Private/GeometryCollection/GeometryCollectionComponent.cpp b/Engine/Source/Runtime/Experimental/GeometryCollectionEngine/Private/GeometryCollection/GeometryCollectionComponent.cpp
--- a/Engine/Source/Runtime/Experimental/GeometryCollectionEngine/Private/GeometryCollection/GeometryCollectionComponent.cpp
+++ b/Engine/Source/Runtime/Experimental/GeometryCollectionEngine/Private/GeometryCollection/GeometryCollectionComponent.cpp
@@ -1833,23 +1833,9 @@ void UGeometryCollectionComponent::RegisterAndInitializePhysicsProxy()
 		SimulationParameters.ObjectType = ObjectType;
 		SimulationParameters.CacheType = CacheParameters.CacheMode;
 		SimulationParameters.ReverseCacheBeginTime = CacheParameters.ReverseCacheBeginTime;
-		SimulationParameters.CollisionData.SaveCollisionData = CacheParameters.SaveCollisionData;
-		SimulationParameters.CollisionData.DoGenerateCollisionData = CacheParameters.DoGenerateCollisionData;
-		SimulationParameters.CollisionData.CollisionDataSizeMax = CacheParameters.CollisionDataSizeMax;
-		SimulationParameters.CollisionData.DoCollisionDataSpatialHash = CacheParameters.DoCollisionDataSpatialHash;
-		SimulationParameters.CollisionData.CollisionDataSpatialHashRadius = CacheParameters.CollisionDataSpatialHashRadius;
-		SimulationParameters.CollisionData.MaxCollisionPerCell = CacheParameters.MaxCollisionPerCell;
-		SimulationParameters.BreakingData.SaveBreakingData = CacheParameters.SaveBreakingData;
-		SimulationParameters.BreakingData.DoGenerateBreakingData = CacheParameters.DoGenerateBreakingData;
-		SimulationParameters.BreakingData.BreakingDataSizeMax = CacheParameters.BreakingDataSizeMax;
-		SimulationParameters.BreakingData.DoBreakingDataSpatialHash = CacheParameters.DoBreakingDataSpatialHash;
-		SimulationParameters.BreakingData.BreakingDataSpatialHashRadius = CacheParameters.BreakingDataSpatialHashRadius;
-		SimulationParameters.BreakingData.MaxBreakingPerCell = CacheParameters.MaxBreakingPerCell;
-		SimulationParameters.TrailingData.SaveTrailingData = CacheParameters.SaveTrailingData;
-		SimulationParameters.TrailingData.DoGenerateTrailingData = CacheParameters.DoGenerateTrailingData;
-		SimulationParameters.TrailingData.TrailingDataSizeMax = CacheParameters.TrailingDataSizeMax;
-		SimulationParameters.TrailingData.TrailingMinSpeedThreshold = CacheParameters.TrailingMinSpeedThreshold;
-		SimulationParameters.TrailingData.TrailingMinVolumeThreshold = CacheParameters.TrailingMinVolumeThreshold;
+		SimulationParameters.bGenerateBreakingData = bNotifyBreaks;
+		SimulationParameters.bGenerateCollisionData = bNotifyCollisions;
+		SimulationParameters.bGenerateTrailingData = bNotifyTrailing;
 		SimulationParameters.RemoveOnFractureEnabled = SimulationParameters.Shared.RemoveOnFractureIndices.Num() > 0;
 		SimulationParameters.WorldTransform = GetComponentToWorld();
 		SimulationParameters.UserData = static_cast<void*>(&PhysicsUserData);
