ChaosCaching
-- Corrected Solver configuration from GeometryCollectionCacheAdapter to enqueue command on the physics thread instead of setting directly from the game thread.

diff --git a/Engine/Plugins/Experimental/ChaosCaching/Source/ChaosCaching/Private/Chaos/Adapters/GeometryCollectionComponentCacheAdapter.cpp b/Engine/Plugins/Experimental/ChaosCaching/Source/ChaosCaching/Private/Chaos/Adapters/GeometryCollectionComponentCacheAdapter.cpp
--- a/Engine/Plugins/Experimental/ChaosCaching/Source/ChaosCaching/Private/Chaos/Adapters/GeometryCollectionComponentCacheAdapter.cpp
+++ b/Engine/Plugins/Experimental/ChaosCaching/Source/ChaosCaching/Private/Chaos/Adapters/GeometryCollectionComponentCacheAdapter.cpp
@@ -685,9 +685,12 @@ namespace Chaos
 		}
 
 		// We need secondary event data to record event information into the cache
-		Solver->SetGenerateBreakingData(true);
-		Solver->SetGenerateCollisionData(true);
-		Solver->SetGenerateTrailingData(true);
+		Solver->EnqueueCommandImmediate([Solver]()
+			{
+				Solver->SetGenerateBreakingData(true);
+				Solver->SetGenerateCollisionData(true);
+				Solver->SetGenerateTrailingData(true);
+			});
 		
 		// We only need to register event handlers once, the first time we initialize.
 		if (CachedData.Num() == 0)
