ChaosCaching
-- Strip  out MassToLocal transform from older GeometryCollection caches

diff --git a/Engine/Plugins/Experimental/ChaosCaching/Source/ChaosCaching/Private/Chaos/Adapters/GeometryCollectionComponentCacheAdapter.cpp b/Engine/Plugins/Experimental/ChaosCaching/Source/ChaosCaching/Private/Chaos/Adapters/GeometryCollectionComponentCacheAdapter.cpp
--- a/Engine/Plugins/Experimental/ChaosCaching/Source/ChaosCaching/Private/Chaos/Adapters/GeometryCollectionComponentCacheAdapter.cpp
+++ b/Engine/Plugins/Experimental/ChaosCaching/Source/ChaosCaching/Private/Chaos/Adapters/GeometryCollectionComponentCacheAdapter.cpp
@@ -490,8 +490,7 @@ namespace Chaos
 				}
 
 				// Evaluated transform is in Cache Manager space with MassToLocal removed. We need it in World space.
-				// Note that before Version 1, MassToLocal is burned into the evaluated transform.
-				FTransform WorldTransform = (InCache->Version < 1) ? EvaluatedTransform : MassToLocal[ParticleIndex] * EvaluatedTransform;
+				FTransform WorldTransform = MassToLocal[ParticleIndex] * EvaluatedTransform;
 				
 				Handle->SetP(WorldTransform.GetTranslation());
 				Handle->SetQ(WorldTransform.GetRotation());
@@ -512,7 +511,7 @@ namespace Chaos
 							// This will result in multiple transform sets happening to the parent but allows us to mostly ignore
 							// that it exists, if it doesn't the child still gets set to the correct position.
 							FTransform ChildTransform = Handle->ChildToParent();
-							FTransform Result = ChildTransform.Inverse() * EvaluatedTransform;
+							FTransform Result = ChildTransform.Inverse() * WorldTransform;
 							Parent->SetP(Result.GetTranslation());
 							Parent->SetX(Result.GetTranslation());
 							Parent->SetQ(Result.GetRotation());
@@ -665,8 +664,6 @@ namespace Chaos
 	{
 		ensure(IsInGameThreadContext());
 		
-		InCache->Version = 1;
-		
 		UGeometryCollectionComponent*    Comp     = CastChecked<UGeometryCollectionComponent>(InComponent);
 		FGeometryCollectionPhysicsProxy* Proxy    = Comp->GetPhysicsProxy();
 
diff --git a/Engine/Plugins/Experimental/ChaosCaching/Source/ChaosCaching/Private/Chaos/Adapters/StaticMeshComponentCacheAdapter.cpp b/Engine/Plugins/Experimental/ChaosCaching/Source/ChaosCaching/Private/Chaos/Adapters/StaticMeshComponentCacheAdapter.cpp
--- a/Engine/Plugins/Experimental/ChaosCaching/Source/ChaosCaching/Private/Chaos/Adapters/StaticMeshComponentCacheAdapter.cpp
+++ b/Engine/Plugins/Experimental/ChaosCaching/Source/ChaosCaching/Private/Chaos/Adapters/StaticMeshComponentCacheAdapter.cpp
@@ -167,7 +167,6 @@ namespace Chaos
 	
 	bool FStaticMeshCacheAdapter::InitializeForRecord(UPrimitiveComponent* InComponent, UChaosCache* InCache)
 	{
-		InCache->Version = 1;
 		return true;
 	}
 
diff --git a/Engine/Plugins/Experimental/ChaosCaching/Source/ChaosCaching/Private/Chaos/ChaosCache.cpp b/Engine/Plugins/Experimental/ChaosCaching/Source/ChaosCaching/Private/Chaos/ChaosCache.cpp
--- a/Engine/Plugins/Experimental/ChaosCaching/Source/ChaosCaching/Private/Chaos/ChaosCache.cpp
+++ b/Engine/Plugins/Experimental/ChaosCaching/Source/ChaosCaching/Private/Chaos/ChaosCache.cpp
@@ -3,12 +3,96 @@
 #include "Chaos/ChaosCache.h"
 #include "Chaos/ChaosCachingPlugin.h"
 #include "Components/PrimitiveComponent.h"
+#include "UObject/UE5MainStreamObjectVersion.h"
+#include "GeometryCollection/GeometryCollectionComponent.h"
+#include "GeometryCollection/GeometryCollectionObject.h"
 
 UChaosCache::UChaosCache()
 	: CurrentRecordCount(0)
 	, CurrentPlaybackCount(0)
+	, bStripMassToLocal(false)
 {
+	// Default Version. If a Version was archived, this will be overwritten during Serialize() 
+	Version = 0;
+}
+
+void UChaosCache::Serialize(FArchive& Ar)
+{
+	Ar.UsingCustomVersion(FUE5MainStreamObjectVersion::GUID);
+
+	Super::Serialize(Ar);
+
+	if (Ar.IsLoading())
+	{ 
+		// Older versions of GeometryCollection caches had MassToLocal transform baked into the stored transforms. This unfortunately means
+		// that evaluating the cache outside the context of the PhysicsThread is unlikely to be accurate. To make this work, we need to
+		// strip the MassToLocal from the existing cached transforms.
+		if ((Ar.CustomVer(FUE5MainStreamObjectVersion::GUID) < FUE5MainStreamObjectVersion::GeometryCollectionCacheRemovesMassToLocal) && (Version == 0))
+		{
+			bStripMassToLocal = true;
+		}
+	}	
+}
+
+void UChaosCache::PostLoad()
+{
+	if (bStripMassToLocal)
+	{ 
+		// Am I a GeometryCollection cache?
+		UObject* DuplicatedTemplate = Spawnable.DuplicatedTemplate;
+		if (UGeometryCollectionComponent* GeometryCollectionComponent = Cast<UGeometryCollectionComponent>(DuplicatedTemplate))
+		{
+			// Get the referenced RestCollection
+			if (const UGeometryCollection* RestCollection = GeometryCollectionComponent->GetRestCollection())
+			{
+				const TSharedPtr<FGeometryCollection, ESPMode::ThreadSafe> GeometryCollection = RestCollection->GetGeometryCollection();
+
+				// Get the MassToLocal transforms
+				if (GeometryCollection->HasAttribute(TEXT("MassToLocal"), FTransformCollection::TransformGroup))
+				{
+					const TManagedArray<FTransform>& CollectionMassToLocal = GeometryCollection->GetAttribute<FTransform>(TEXT("MassToLocal"), FTransformCollection::TransformGroup);
+				
+					// Strip out the MassToLocal transforms from all cached transforms
+					int32 NumTracks = TrackToParticle.Num();
+					int32 NumParticles = CollectionMassToLocal.Num();
+					for (int32 TrackIdx = 0; TrackIdx < NumTracks; ++TrackIdx)
+					{
+						int32 ParticleIdx = TrackToParticle[TrackIdx];
+						if (ParticleIdx < NumParticles)
+						{ 
+							const FTransform& MassToLocal = CollectionMassToLocal[ParticleIdx];
+
+							FParticleTransformTrack& TransformData = ParticleTracks[TrackIdx].TransformData;
+							FRawAnimSequenceTrack& AnimTrack = TransformData.RawTransformTrack;
+
+							if (ensure((AnimTrack.PosKeys.Num() == AnimTrack.RotKeys.Num()) && (AnimTrack.RotKeys.Num() == AnimTrack.ScaleKeys.Num())))
+							{ 
+								int32 NumKeys = AnimTrack.PosKeys.Num();
+								for (int32 KeyIdx = 0; KeyIdx < NumKeys; ++KeyIdx)
+								{
+									FTransform MassTransform(AnimTrack.RotKeys[KeyIdx], AnimTrack.PosKeys[KeyIdx], AnimTrack.ScaleKeys[KeyIdx]);
+									FTransform LocalTransform = MassToLocal.Inverse() * MassTransform;
+									FQuat LocalRotation = LocalTransform.GetRotation();
+									FVector LocalTranslation = LocalTransform.GetTranslation();
+
+									AnimTrack.RotKeys[KeyIdx] = LocalRotation;
+									AnimTrack.PosKeys[KeyIdx] = LocalTranslation;
+								}
+							}
+						}
+					}
+				
+					bStripMassToLocal = false;
+										
+				}
+			}
+		}
+	}
 
+	// Up-to-date now
+	Version = CurrentVersion;
+
+	Super::PostLoad();
 }
 
 void UChaosCache::FlushPendingFrames()
@@ -161,6 +245,9 @@ void UChaosCache::EndRecord(FCacheUserToken& InOutToken)
 		InOutToken.bIsOpen = false;
 		InOutToken.Owner = nullptr;
 		CurrentRecordCount--;
+
+		// Set the Version appropriately
+		Version = CurrentVersion;
 	}
 	else
 	{
diff --git a/Engine/Plugins/Experimental/ChaosCaching/Source/ChaosCaching/Public/Chaos/ChaosCache.h b/Engine/Plugins/Experimental/ChaosCaching/Source/ChaosCaching/Public/Chaos/ChaosCache.h
--- a/Engine/Plugins/Experimental/ChaosCaching/Source/ChaosCaching/Public/Chaos/ChaosCache.h
+++ b/Engine/Plugins/Experimental/ChaosCaching/Source/ChaosCaching/Public/Chaos/ChaosCache.h
@@ -252,6 +252,11 @@ public:
 
 	UChaosCache();
 
+	//~ UObject interface
+	virtual void Serialize(FArchive& Ar) override;
+	virtual void PostLoad() override;
+	//~ UObject interface END
+
 	/**
 	 * As we record post-simulate of physics, we're almost always taking data from a non-main thread (physics thread).
 	 * Because of this we can't directly write into the cache, but instead into a pending frame queue that needs to be
@@ -348,14 +353,6 @@ public:
 	UPROPERTY()
 	TMap<FName, FRichCurve> CurveData;
 
-	// Version for introducing new features to caches. If cache recording changes content or format of any caches,
-	// the version should be incremented and cache playback should accommodate older cache versions.
-	// The cache version is during InitializeForRecord in the appropriate adapter.
-	// 0
-	// 1 : Removed MassToLocal transform for GeometryCollection caches
-	UPROPERTY()
-	int32 Version;
-
 	template<typename T>
 	FCacheEventTrack& FindOrAddEventTrack(FName InName)
 	{
@@ -382,10 +379,21 @@ private:
 	UPROPERTY()
 	FGuid AdapterGuid;
 
+	/** Version for controlling conditioning of older caches to work with current system. Newly created caches should always be saved as CurrentVersion. */
+	UPROPERTY()
+	int32 Version;
+
+	// Version 0 : Pre versioning.
+	// Version 1 : Introduction of actor space transforms & removal of baked MassTOLocal transform in GeometryCollections.
+	static constexpr int32 CurrentVersion = 1;
+
 	/** Pending writes from all threads to be consumed on the game thread, triggered by the recording cache manager */
 	TQueue<FPendingFrameWrite, EQueueMode::Mpsc> PendingWrites;
 
 	/** Counts for current number of users, should only ever have one recorder, and if we do no playbacks */
 	std::atomic<int32> CurrentRecordCount;
 	std::atomic<int32> CurrentPlaybackCount;
+
+	/** Indicates that we need to strip MassToLocal before playing the cache. */
+	bool bStripMassToLocal;
 };
diff --git a/Engine/Source/Runtime/Core/Public/UObject/UE5MainStreamObjectVersion.h b/Engine/Source/Runtime/Core/Public/UObject/UE5MainStreamObjectVersion.h
--- a/Engine/Source/Runtime/Core/Public/UObject/UE5MainStreamObjectVersion.h
+++ b/Engine/Source/Runtime/Core/Public/UObject/UE5MainStreamObjectVersion.h
@@ -150,13 +150,16 @@ struct CORE_API FUE5MainStreamObjectVersion
 
 		// Removed the AtmosphericFog component with conversion to SkyAtmosphere component
 		RemovedAtmosphericFog,
-
+		
 		// The SkyAtmosphere now light up the heightfog by default, and by default the height fog has a black color.
 		SkyAtmosphereAffectsHeightFogWithBetterDefault,
 
 		// Ordering of samples in BlendSpace
 		BlendSpaceSampleOrdering,
 
+		// No longer bake MassToLocal transform into recorded transform data in GeometryCollection caching
+		GeometryCollectionCacheRemovesMassToLocal,
+
 		// -----<new versions can be added above this line>-------------------------------------------------
 		VersionPlusOne,
 		LatestVersion = VersionPlusOne - 1
