ChaosCache
-- Flush PT command queue before issuing config commands
-- Corrected an issue with resetting cached events when more than one manager are present.

diff --git a/Engine/Plugins/Experimental/ChaosCaching/Source/ChaosCaching/Private/Chaos/Adapters/GeometryCollectionComponentCacheAdapter.cpp b/Engine/Plugins/Experimental/ChaosCaching/Source/ChaosCaching/Private/Chaos/Adapters/GeometryCollectionComponentCacheAdapter.cpp
--- a/Engine/Plugins/Experimental/ChaosCaching/Source/ChaosCaching/Private/Chaos/Adapters/GeometryCollectionComponentCacheAdapter.cpp
+++ b/Engine/Plugins/Experimental/ChaosCaching/Source/ChaosCaching/Private/Chaos/Adapters/GeometryCollectionComponentCacheAdapter.cpp
@@ -591,11 +591,6 @@ namespace Chaos
 		return nullptr;
 	}
 
-	void FGeometryCollectionCacheAdapter::Initialize()
-	{
-		CachedData.Empty();
-	}
-
 	void FGeometryCollectionCacheAdapter::SetRestState(UPrimitiveComponent* InComponent, UChaosCache* InCache, const FTransform& InRootTransform, Chaos::FReal InTime) const
 	{
 		// Caches recorded previous to Version 1 may not scrub correctly as the MassToLocal transform has been burned in.
@@ -681,6 +676,10 @@ namespace Chaos
 			return false;
 		}
 
+		// In case commands have been issued that conflict with our requirement to generate data, flush the queue
+		Solver->AdvanceAndDispatch_External(0);
+		Solver->WaitOnPendingTasks_External();
+
 		// We need secondary event data to record event information into the cache
 		Solver->EnqueueCommandImmediate([Solver]()
 			{
@@ -717,7 +716,10 @@ namespace Chaos
 			TrailingDataArray = nullptr;
 		}
 
-		CachedData.Add(Proxy, FCachedEventData());
+		FCachedEventData& CachedEventData = CachedData.FindOrAdd(Proxy);
+		CachedEventData.ProxyBreakingDataIndices = nullptr;
+		CachedEventData.ProxyCollisionDataIndices = nullptr;
+		CachedEventData.ProxyTrailingDataIndices = nullptr;
 
 		return true;
 	}
diff --git a/Engine/Plugins/Experimental/ChaosCaching/Source/ChaosCaching/Public/Chaos/Adapters/GeometryCollectionComponentCacheAdapter.h b/Engine/Plugins/Experimental/ChaosCaching/Source/ChaosCaching/Public/Chaos/Adapters/GeometryCollectionComponentCacheAdapter.h
--- a/Engine/Plugins/Experimental/ChaosCaching/Source/ChaosCaching/Public/Chaos/Adapters/GeometryCollectionComponentCacheAdapter.h
+++ b/Engine/Plugins/Experimental/ChaosCaching/Source/ChaosCaching/Public/Chaos/Adapters/GeometryCollectionComponentCacheAdapter.h
@@ -220,7 +220,7 @@ namespace Chaos
 			: ProxyBreakingDataIndices(nullptr)
 			, ProxyCollisionDataIndices(nullptr)
 			, ProxyTrailingDataIndices(nullptr)
-		{}
+		{ }
 
 		const TArray<int32>* ProxyBreakingDataIndices;
 		const TArray<int32>* ProxyCollisionDataIndices;
@@ -240,7 +240,6 @@ namespace Chaos
 		FGuid                  GetGuid() const override;
 		Chaos::FPhysicsSolver* GetComponentSolver(UPrimitiveComponent* InComponent) const override;
 		bool                   ValidForPlayback(UPrimitiveComponent* InComponent, UChaosCache* InCache) const override;
-		void				   Initialize() override;
 		void				   SetRestState(UPrimitiveComponent* InComponent, UChaosCache* InCache, const FTransform& InRootTransform, Chaos::FReal InTime) const override;
 		bool                   InitializeForRecord(UPrimitiveComponent* InComponent, UChaosCache* InCache) override;
 		bool                   InitializeForPlayback(UPrimitiveComponent* InComponent, UChaosCache* InCache, float InTime) override;
