ChaosCaching
-- Pull PhyScene pointer from solver so that recording breaks in geometry collection works with custom solver.

diff --git a/Engine/Plugins/Experimental/ChaosCaching/Source/ChaosCaching/Private/Chaos/Adapters/GeometryCollectionComponentCacheAdapter.cpp b/Engine/Plugins/Experimental/ChaosCaching/Source/ChaosCaching/Private/Chaos/Adapters/GeometryCollectionComponentCacheAdapter.cpp
--- a/Engine/Plugins/Experimental/ChaosCaching/Source/ChaosCaching/Private/Chaos/Adapters/GeometryCollectionComponentCacheAdapter.cpp
+++ b/Engine/Plugins/Experimental/ChaosCaching/Source/ChaosCaching/Private/Chaos/Adapters/GeometryCollectionComponentCacheAdapter.cpp
@@ -53,8 +53,6 @@ namespace Chaos
 		UGeometryCollectionComponent*    Comp  = CastChecked<UGeometryCollectionComponent>(InComp);
 		FGeometryCollectionPhysicsProxy* Proxy = Comp->GetPhysicsProxy();
 		
-		FPhysScene& Scene = *(Comp->GetWorld()->GetPhysicsScene());
-		
 		if(!Proxy)
 		{
 			return;
@@ -77,6 +75,12 @@ namespace Chaos
 			return;
 		}
 
+		const FPhysScene* Scene = static_cast<FPhysScene*>(Solver->PhysSceneHack);
+		if (!Scene)
+		{
+			return;
+		}
+
 		FGeometryDynamicCollection&            Collection       = Proxy->GetPhysicsCollection();
 		const TManagedArray<FTransform>&       MassToLocal		= Collection.MassToLocal;
 		const TManagedArray<int32>&            Parents			= RestCollection->Parent;
@@ -100,7 +104,7 @@ namespace Chaos
 		for(const FBreakingData& Break : Breaks)
 		{
 			// Is the proxy pending destruction? If they are no longer tracked by the PhysScene, the proxy is deleted or pending deletion.
-			if (Scene.GetOwningComponent<UPrimitiveComponent>(Break.Proxy) == nullptr)
+			if (Scene->GetOwningComponent<UPrimitiveComponent>(Break.Proxy) == nullptr)
 			{
 				continue;
 			}
