ChaosCaching
-- Check for null proxy before accessing for play/record.

diff --git a/Engine/Plugins/Experimental/ChaosCaching/Source/ChaosCaching/Private/Chaos/Adapters/StaticMeshComponentCacheAdapter.cpp b/Engine/Plugins/Experimental/ChaosCaching/Source/ChaosCaching/Private/Chaos/Adapters/StaticMeshComponentCacheAdapter.cpp
--- a/Engine/Plugins/Experimental/ChaosCaching/Source/ChaosCaching/Private/Chaos/Adapters/StaticMeshComponentCacheAdapter.cpp
+++ b/Engine/Plugins/Experimental/ChaosCaching/Source/ChaosCaching/Private/Chaos/Adapters/StaticMeshComponentCacheAdapter.cpp
@@ -38,6 +38,11 @@ namespace Chaos
 
 	void RecordToCacheInternal(FSingleParticlePhysicsProxy* InProxy, const FTransform& InRootTransform, FPendingFrameWrite& OutFrame, Chaos::FReal InTime)
 	{
+		if (!InProxy)
+		{
+			return;
+		}
+
 		if(TGeometryParticleHandle<FReal,3>* Handle = InProxy->GetHandle_LowLevel())
 		{
 			if(FPBDRigidParticleHandle* AsRigid = Handle->CastToRigidParticle())
@@ -66,7 +71,7 @@ namespace Chaos
 
 	void PlayFromCacheInternal(FSingleParticlePhysicsProxy* InProxy, UChaosCache* InCache, FPlaybackTickRecord& TickRecord, TArray<TPBDRigidParticleHandle<Chaos::FReal, 3>*>& OutUpdatedRigids) 
 	{
-		if(!InCache || InCache->GetDuration() == 0.0f)
+		if(!InCache || InCache->GetDuration() == 0.0f || !InProxy)
 		{
 			return;
 		}
