Chaos : CacheManager
- Check for valid actor handle before initializing in the cache.
- preflight was hung on scheduling Win64 BootTest (id 611447e09c7bb10001f1eda0)

diff --git a/Engine/Plugins/Experimental/ChaosCaching/Source/ChaosCaching/Private/Chaos/Adapters/StaticMeshComponentCacheAdapter.cpp b/Engine/Plugins/Experimental/ChaosCaching/Source/ChaosCaching/Private/Chaos/Adapters/StaticMeshComponentCacheAdapter.cpp
--- a/Engine/Plugins/Experimental/ChaosCaching/Source/ChaosCaching/Private/Chaos/Adapters/StaticMeshComponentCacheAdapter.cpp
+++ b/Engine/Plugins/Experimental/ChaosCaching/Source/ChaosCaching/Private/Chaos/Adapters/StaticMeshComponentCacheAdapter.cpp
@@ -179,7 +179,7 @@ namespace Chaos
 	{
 #if WITH_CHAOS
 
-		if(Cast<UStaticMeshComponent>(InComponent))
+		if(Cast<UStaticMeshComponent>(InComponent) && InComponent->GetBodyInstance() && InComponent->GetBodyInstance()->ActorHandle)
 		{
 			FPhysInterface_Chaos::SetIsKinematic_AssumesLocked(InComponent->GetBodyInstance()->ActorHandle, true);
 		}
