ChaosCaching
-- don't initialize physics simulation on component unless cache player is set to Play.

diff --git a/Engine/Plugins/Experimental/ChaosCaching/Source/ChaosCaching/Private/Chaos/CacheManagerActor.cpp b/Engine/Plugins/Experimental/ChaosCaching/Source/ChaosCaching/Private/Chaos/CacheManagerActor.cpp
--- a/Engine/Plugins/Experimental/ChaosCaching/Source/ChaosCaching/Private/Chaos/CacheManagerActor.cpp
+++ b/Engine/Plugins/Experimental/ChaosCaching/Source/ChaosCaching/Private/Chaos/CacheManagerActor.cpp
@@ -317,14 +317,17 @@ void AChaosCacheManager::BeginPlay()
 
 					if(Token.IsOpen() && CurrAdapter->ValidForPlayback(Comp, PlayCache))
 					{
-						SolverData->PlaybackIndices.Add(Index);
-						SolverData->PlaybackTickRecords.AddDefaulted();
-						SolverData->PlaybackTickRecords.Last().SetSpaceTransform(GetTransform());
-						Observed.Cache = PlayCache;
-						Observed.TickRecord.SetSpaceTransform(GetTransform());
-						Observed.TickRecord.SetLastTime(StartTime);
-						OpenPlaybackCaches.Add(TTuple<FCacheUserToken, UChaosCache*>(MoveTemp(Token), Observed.Cache));
-						CurrAdapter->InitializeForPlayback(Comp, Observed.Cache, StartTime);
+						if (CacheMode == ECacheMode::Play)
+						{
+							SolverData->PlaybackIndices.Add(Index);
+							SolverData->PlaybackTickRecords.AddDefaulted();
+							SolverData->PlaybackTickRecords.Last().SetSpaceTransform(GetTransform());
+							Observed.Cache = PlayCache;
+							Observed.TickRecord.SetSpaceTransform(GetTransform());
+							Observed.TickRecord.SetLastTime(StartTime);
+							OpenPlaybackCaches.Add(TTuple<FCacheUserToken, UChaosCache*>(MoveTemp(Token), Observed.Cache));
+							CurrAdapter->InitializeForPlayback(Comp, Observed.Cache, StartTime);
+						}
 					}
 					else
 					{
