Chaos Caching
-- fix to ensure that partially broken clusters are recorded properly

diff --git a/Engine/Plugins/Experimental/ChaosCaching/Source/ChaosCaching/Private/Chaos/Adapters/GeometryCollectionComponentCacheAdapter.cpp b/Engine/Plugins/Experimental/ChaosCaching/Source/ChaosCaching/Private/Chaos/Adapters/GeometryCollectionComponentCacheAdapter.cpp
--- a/Engine/Plugins/Experimental/ChaosCaching/Source/ChaosCaching/Private/Chaos/Adapters/GeometryCollectionComponentCacheAdapter.cpp
+++ b/Engine/Plugins/Experimental/ChaosCaching/Source/ChaosCaching/Private/Chaos/Adapters/GeometryCollectionComponentCacheAdapter.cpp
@@ -84,6 +84,7 @@ namespace Chaos
 		FGeometryDynamicCollection&            Collection       = Proxy->GetPhysicsCollection();
 		const TManagedArray<FTransform>&       MassToLocal		= Collection.MassToLocal;
 		const TManagedArray<int32>&            Parents			= RestCollection->Parent;
+		const TManagedArray<TSet<int32>>&	   Children         = RestCollection->Children;
 		const TArray<FBreakingData>&		   Breaks           = Solver->GetEvolution()->GetRigidClustering().GetAllClusterBreakings();
 		const FTransform					   WorldToActor		= InRootTransform.Inverse();
 
@@ -115,11 +116,14 @@ namespace Chaos
 
 				if(ConcreteProxy == Proxy)
 				{
-					// The break particle belongs to our proxy
-					RelatedBreaks.Add(Break.TransformGroupIndex);
-
 					// Record the cluster it belonged to.
 					ReleasedClusters.Add(Parents[Break.TransformGroupIndex]);
+
+					// Force a break on all children of the cluster
+					for (int32 ChildToBreak : Children[Parents[Break.TransformGroupIndex]])
+					{
+						RelatedBreaks.AddUnique(ChildToBreak);
+					}
 				}
 			}
 		}
