Fix IsInPhysicsThreadContext() ensure firing in test/shipping builds due to missing #if wrapper.

diff --git a/Engine/Plugins/Experimental/ChaosCaching/Source/ChaosCaching/Private/Chaos/Adapters/GeometryCollectionComponentCacheAdapter.cpp b/Engine/Plugins/Experimental/ChaosCaching/Source/ChaosCaching/Private/Chaos/Adapters/GeometryCollectionComponentCacheAdapter.cpp
--- a/Engine/Plugins/Experimental/ChaosCaching/Source/ChaosCaching/Private/Chaos/Adapters/GeometryCollectionComponentCacheAdapter.cpp
+++ b/Engine/Plugins/Experimental/ChaosCaching/Source/ChaosCaching/Private/Chaos/Adapters/GeometryCollectionComponentCacheAdapter.cpp
@@ -45,7 +45,11 @@ namespace Chaos
 	
 	void FGeometryCollectionCacheAdapter::Record_PostSolve(UPrimitiveComponent* InComp, const FTransform& InRootTransform, FPendingFrameWrite& OutFrame, Chaos::FReal InTime) const
 	{
+
+#if PHYSICS_THREAD_CONTEXT
 		ensure(IsInPhysicsThreadContext());
+#endif
+
 
 		using FClusterParticle = Chaos::FPBDRigidClusteredParticleHandle;
 		using FRigidParticle = Chaos::FPBDRigidParticleHandle;
@@ -234,7 +238,9 @@ namespace Chaos
 															FPlaybackTickRecord&							   TickRecord,
 															TArray<TPBDRigidParticleHandle<Chaos::FReal, 3>*>& OutUpdatedRigids) const
 	{
+#if PHYSICS_THREAD_CONTEXT
 		ensure(IsInPhysicsThreadContext());
+#endif
 		
 		using FClusterParticle = Chaos::FPBDRigidClusteredParticleHandle;
 		using FRigidParticle = Chaos::FPBDRigidParticleHandle;
@@ -603,7 +609,10 @@ namespace Chaos
 	{
 		// Caches recorded previous to Version 1 may not scrub correctly as the MassToLocal transform has been burned in.
 		
+#if PHYSICS_THREAD_CONTEXT
 		ensure(IsInGameThreadContext());
+#endif
+
 		
 		if (!InCache || InCache->GetDuration() == 0.0f)
 		{
@@ -665,7 +674,9 @@ namespace Chaos
 
 	bool FGeometryCollectionCacheAdapter::InitializeForRecord(UPrimitiveComponent* InComponent, UChaosCache* InCache)
 	{
+#if PHYSICS_THREAD_CONTEXT
 		ensure(IsInGameThreadContext());
+#endif
 		
 		UGeometryCollectionComponent*    Comp     = CastChecked<UGeometryCollectionComponent>(InComponent);
 		FGeometryCollectionPhysicsProxy* Proxy    = Comp->GetPhysicsProxy();
@@ -734,7 +745,9 @@ namespace Chaos
 
 	bool FGeometryCollectionCacheAdapter::InitializeForPlayback(UPrimitiveComponent* InComponent, UChaosCache* InCache, float InTime)
 	{
+#if PHYSICS_THREAD_CONTEXT
 		ensure(IsInGameThreadContext());
+#endif
 
 		if (UGeometryCollectionComponent* Comp = CastChecked<UGeometryCollectionComponent>(InComponent))
 		{
@@ -822,7 +835,9 @@ namespace Chaos
 
 	void FGeometryCollectionCacheAdapter::HandleBreakingEvents(const Chaos::FBreakingEventData& Event)
 	{
+#if PHYSICS_THREAD_CONTEXT
 		ensure(IsInGameThreadContext());
+#endif
 		
 		BreakingDataArray = &Event.BreakingData.AllBreakingsArray;
 
@@ -841,7 +856,9 @@ namespace Chaos
 
 	void FGeometryCollectionCacheAdapter::HandleCollisionEvents(const Chaos::FCollisionEventData& Event)
 	{
+#if PHYSICS_THREAD_CONTEXT
 		ensure(IsInGameThreadContext());
+#endif
 		
 		CollisionDataArray = &Event.CollisionData.AllCollisionsArray;
 
@@ -860,7 +877,9 @@ namespace Chaos
 
 	void FGeometryCollectionCacheAdapter::HandleTrailingEvents(const Chaos::FTrailingEventData& Event)
 	{
+#if PHYSICS_THREAD_CONTEXT
 		ensure(IsInGameThreadContext());
+#endif
 		
 		TrailingDataArray = &Event.TrailingData.AllTrailingsArray;
 
diff --git a/Engine/Plugins/Experimental/ChaosCaching/Source/ChaosCaching/Private/Chaos/Adapters/StaticMeshComponentCacheAdapter.cpp b/Engine/Plugins/Experimental/ChaosCaching/Source/ChaosCaching/Private/Chaos/Adapters/StaticMeshComponentCacheAdapter.cpp
--- a/Engine/Plugins/Experimental/ChaosCaching/Source/ChaosCaching/Private/Chaos/Adapters/StaticMeshComponentCacheAdapter.cpp
+++ b/Engine/Plugins/Experimental/ChaosCaching/Source/ChaosCaching/Private/Chaos/Adapters/StaticMeshComponentCacheAdapter.cpp
@@ -177,7 +177,9 @@ namespace Chaos
 
 	bool FStaticMeshCacheAdapter::InitializeForPlayback(UPrimitiveComponent* InComponent, UChaosCache* InCache, float InTime)
 	{
+#if PHYSICS_THREAD_CONTEXT
 		ensure(IsInGameThreadContext());
+#endif
 
 #if WITH_CHAOS
 
diff --git a/Engine/Source/Programs/HeadlessChaos/Private/HeadlessChaosTestRewind.cpp b/Engine/Source/Programs/HeadlessChaos/Private/HeadlessChaosTestRewind.cpp
--- a/Engine/Source/Programs/HeadlessChaos/Private/HeadlessChaosTestRewind.cpp
+++ b/Engine/Source/Programs/HeadlessChaos/Private/HeadlessChaosTestRewind.cpp
@@ -2453,8 +2453,9 @@ namespace ChaosTest {
 			}
 			
 			const int RewindStep = 5;
-
+#if PHYSICS_THREAD_CONTEXT
 			FPhysicsThreadContextScope Scope(true);
+#endif
 			auto& Particle = *Proxy->GetPhysicsThreadAPI();
 			FRewindData* RewindData = Solver->GetRewindData();
 			//EXPECT_TRUE(RewindData->RewindToFrame(RewindStep));
@@ -2550,7 +2551,9 @@ namespace ChaosTest {
 			
 			const int RewindStep = 5;
 
+#if PHYSICS_THREAD_CONTEXT
 			FPhysicsThreadContextScope Scope(true);
+#endif
 			FRewindData* RewindData = Solver->GetRewindData();
 			//EXPECT_TRUE(RewindData->RewindToFrame(RewindStep));
 			auto& Particle = *Proxy->GetPhysicsThreadAPI();
@@ -2740,7 +2743,9 @@ namespace ChaosTest {
 			const int RewindStep = 5;
 			auto& Particle = *Proxy->GetPhysicsThreadAPI();
 
+#if PHYSICS_THREAD_CONTEXT
 			FPhysicsThreadContextScope Scope(true);
+#endif
 			{
 				FRewindData* RewindData = Solver->GetRewindData();
 				//EXPECT_TRUE(RewindData->RewindToFrame(RewindStep));
@@ -2932,7 +2937,9 @@ namespace ChaosTest {
 
 			const int RewindStep = 0;
 
+#if PHYSICS_THREAD_CONTEXT
 			FPhysicsThreadContextScope Scope(true);
+#endif
 			auto& Dynamic = *DynamicProxy->GetPhysicsThreadAPI();
 			auto& Kinematic = *KinematicProxy->GetPhysicsThreadAPI();
 
@@ -3022,7 +3029,9 @@ namespace ChaosTest {
 				EXPECT_LE(Dynamic.X()[2], 6);
 			}
 
+#if PHYSICS_THREAD_CONTEXT
 			FPhysicsThreadContextScope Scope(true);
+#endif
 			auto& Dynamic = *DynamicProxy->GetPhysicsThreadAPI();
 			auto& Kinematic = *KinematicProxy->GetPhysicsThreadAPI();
 			const int RewindStep = 0;
@@ -3370,7 +3379,9 @@ namespace ChaosTest {
 				EXPECT_GE(Dynamic.X()[2], 10);
 			}
 
+#if PHYSICS_THREAD_CONTEXT
 			FPhysicsThreadContextScope Scope(true);
+#endif
 			auto& Dynamic = *DynamicProxy->GetPhysicsThreadAPI();
 			auto& Kinematic = *KinematicProxy->GetPhysicsThreadAPI();
 
@@ -3479,7 +3490,9 @@ namespace ChaosTest {
 			EXPECT_LE(Dynamic.X()[2], 12);
 		}
 
+#if PHYSICS_THREAD_CONTEXT
 		FPhysicsThreadContextScope Scope(true);
+#endif
 		auto& Dynamic = *DynamicProxy->GetPhysicsThreadAPI();
 		auto& Kinematic = *KinematicProxy->GetPhysicsThreadAPI();
 
@@ -3586,7 +3599,9 @@ namespace ChaosTest {
 			}
 		}
 
+#if PHYSICS_THREAD_CONTEXT
 		FPhysicsThreadContextScope Scope(true);
+#endif
 		auto& Dynamic = *DynamicProxy->GetPhysicsThreadAPI();
 		auto& Kinematic = *KinematicProxy->GetPhysicsThreadAPI();
 
@@ -3675,7 +3690,9 @@ namespace ChaosTest {
 
 			const int32 RewindStep = 3;
 
+#if PHYSICS_THREAD_CONTEXT
 			FPhysicsThreadContextScope Scope(true);
+#endif
 			FRewindData* RewindData = UnitTest.Solver->GetRewindData();
 			//EXPECT_TRUE(RewindData->RewindToFrame(RewindStep));
 
diff --git a/Engine/Source/Runtime/Experimental/Chaos/Private/Chaos/Framework/Parallel.cpp b/Engine/Source/Runtime/Experimental/Chaos/Private/Chaos/Framework/Parallel.cpp
--- a/Engine/Source/Runtime/Experimental/Chaos/Private/Chaos/Framework/Parallel.cpp
+++ b/Engine/Source/Runtime/Experimental/Chaos/Private/Chaos/Framework/Parallel.cpp
@@ -25,7 +25,15 @@ void Chaos::PhysicsParallelFor(int32 InNum, TFunctionRef<void(int32)> InCallable
 	using namespace Chaos;
 	// Passthrough for now, except with global flag to disable parallel
 	
-	auto PassThrough = [InCallable, bIsInPhysicsSimContext = IsInPhysicsThreadContext(), bIsInGameThreadContext = IsInGameThreadContext()](int32 Idx)
+#if PHYSICS_THREAD_CONTEXT
+	const bool bIsInPhysicsSimContext = IsInPhysicsThreadContext();
+	const bool bIsInGameThreadContext = IsInGameThreadContext();
+#else
+	const bool bIsInPhysicsSimContext = false;
+	const bool bIsInGameThreadContext = false;
+#endif
+
+	auto PassThrough = [InCallable, bIsInPhysicsSimContext, bIsInGameThreadContext](int32 Idx)
 	{
 #if PHYSICS_THREAD_CONTEXT
 		FPhysicsThreadContextScope PTScope(bIsInPhysicsSimContext);
diff --git a/Engine/Source/Runtime/Experimental/Chaos/Private/EventDefaults.cpp b/Engine/Source/Runtime/Experimental/Chaos/Private/EventDefaults.cpp
--- a/Engine/Source/Runtime/Experimental/Chaos/Private/EventDefaults.cpp
+++ b/Engine/Source/Runtime/Experimental/Chaos/Private/EventDefaults.cpp
@@ -238,7 +238,10 @@ namespace Chaos
 		(const Chaos::FPBDRigidsSolver* Solver, FBreakingEventData& BreakingEventData)
 		{
 			check(Solver);
+
+#if PHYSICS_THREAD_CONTEXT
 			ensure(IsInPhysicsThreadContext());
+#endif
 
 			SCOPE_CYCLE_COUNTER(STAT_GatherBreakingEvent);
 
@@ -310,7 +313,10 @@ namespace Chaos
 		(const Chaos::FPBDRigidsSolver* Solver, FTrailingEventData& TrailingEventData)
 		{
 			check(Solver);
+
+#if PHYSICS_THREAD_CONTEXT
 			ensure(IsInPhysicsThreadContext());
+#endif
 
 			// #todo: This isn't working - SolverActor parameters are set on a solver but it is currently a different solver that is simulating!!
 			if (!Solver->GetEventFilters()->IsTrailingEventEnabled())
@@ -405,7 +411,10 @@ namespace Chaos
 		(const Chaos::FPBDRigidsSolver* Solver, FSleepingEventData& SleepingEventData)
 		{
 			check(Solver);
+
+#if PHYSICS_THREAD_CONTEXT
 			ensure(IsInPhysicsThreadContext());
+#endif
 
 			SCOPE_CYCLE_COUNTER(STAT_GatherSleepingEvent);
 
@@ -446,7 +455,9 @@ namespace Chaos
 		(const Chaos::FPBDRigidsSolver* Solver, FRemovalEventData& RemovalEventData)
 			{
 				check(Solver);
+#if PHYSICS_THREAD_CONTEXT
 				ensure(IsInPhysicsThreadContext());
+#endif
 
 				
 				FRemovalDataArray& AllRemovalDataArray = RemovalEventData.RemovalData.AllRemovalArray;
diff --git a/Engine/Source/Runtime/Experimental/Chaos/Private/PBDRigidsSolver.cpp b/Engine/Source/Runtime/Experimental/Chaos/Private/PBDRigidsSolver.cpp
--- a/Engine/Source/Runtime/Experimental/Chaos/Private/PBDRigidsSolver.cpp
+++ b/Engine/Source/Runtime/Experimental/Chaos/Private/PBDRigidsSolver.cpp
@@ -643,7 +643,9 @@ namespace Chaos
 
 	void FPBDRigidsSolver::SetSuspensionTargetOnPhysicsThread(Chaos::FSuspensionConstraint* GTConstraint, const FVector& TargetPos, bool Enabled)
 	{
+#if PHYSICS_THREAD_CONTEXT
 		ensure(IsInPhysicsThreadContext());
+#endif
 		FSuspensionConstraintPhysicsProxy* SuspensionProxy = GTConstraint->GetProxy<FSuspensionConstraintPhysicsProxy>();
 		check(SuspensionProxy);
 		SuspensionProxy->UpdateTargetOnPhysicsThread(this, TargetPos, Enabled);
diff --git a/Engine/Source/Runtime/Experimental/Chaos/Private/RewindData.cpp b/Engine/Source/Runtime/Experimental/Chaos/Private/RewindData.cpp
--- a/Engine/Source/Runtime/Experimental/Chaos/Private/RewindData.cpp
+++ b/Engine/Source/Runtime/Experimental/Chaos/Private/RewindData.cpp
@@ -129,7 +129,10 @@ bool FRewindData::RewindToFrame(int32 Frame)
 {
 	QUICK_SCOPE_CYCLE_COUNTER(RewindToFrame);
 
+#if PHYSICS_THREAD_CONTEXT
 	ensure(IsInPhysicsThreadContext());
+#endif
+
 	//Can't go too far back
 	const int32 EarliestFrame = GetEarliestFrame_Internal();
 	if(Frame < EarliestFrame)
