Fix IsInPhysicsContextEnsures firing, add EnsureIsInPhysicsContextEnsures helper.

diff --git a/Engine/Plugins/Experimental/ChaosCaching/Source/ChaosCaching/Private/Chaos/Adapters/GeometryCollectionComponentCacheAdapter.cpp b/Engine/Plugins/Experimental/ChaosCaching/Source/ChaosCaching/Private/Chaos/Adapters/GeometryCollectionComponentCacheAdapter.cpp
--- a/Engine/Plugins/Experimental/ChaosCaching/Source/ChaosCaching/Private/Chaos/Adapters/GeometryCollectionComponentCacheAdapter.cpp
+++ b/Engine/Plugins/Experimental/ChaosCaching/Source/ChaosCaching/Private/Chaos/Adapters/GeometryCollectionComponentCacheAdapter.cpp
@@ -46,9 +46,7 @@ namespace Chaos
 	void FGeometryCollectionCacheAdapter::Record_PostSolve(UPrimitiveComponent* InComp, const FTransform& InRootTransform, FPendingFrameWrite& OutFrame, Chaos::FReal InTime) const
 	{
 
-#if PHYSICS_THREAD_CONTEXT
-		ensure(IsInPhysicsThreadContext());
-#endif
+		EnsureIsInPhysicsThreadContext();
 
 
 		using FClusterParticle = Chaos::FPBDRigidClusteredParticleHandle;
@@ -238,9 +236,7 @@ namespace Chaos
 															FPlaybackTickRecord&							   TickRecord,
 															TArray<TPBDRigidParticleHandle<Chaos::FReal, 3>*>& OutUpdatedRigids) const
 	{
-#if PHYSICS_THREAD_CONTEXT
-		ensure(IsInPhysicsThreadContext());
-#endif
+		EnsureIsInPhysicsThreadContext();
 		
 		using FClusterParticle = Chaos::FPBDRigidClusteredParticleHandle;
 		using FRigidParticle = Chaos::FPBDRigidParticleHandle;
@@ -609,9 +605,7 @@ namespace Chaos
 	{
 		// Caches recorded previous to Version 1 may not scrub correctly as the MassToLocal transform has been burned in.
 		
-#if PHYSICS_THREAD_CONTEXT
-		ensure(IsInGameThreadContext());
-#endif
+		EnsureIsInGameThreadContext();
 
 		
 		if (!InCache || InCache->GetDuration() == 0.0f)
@@ -674,9 +668,7 @@ namespace Chaos
 
 	bool FGeometryCollectionCacheAdapter::InitializeForRecord(UPrimitiveComponent* InComponent, UChaosCache* InCache)
 	{
-#if PHYSICS_THREAD_CONTEXT
-		ensure(IsInGameThreadContext());
-#endif
+		EnsureIsInGameThreadContext();
 		
 		UGeometryCollectionComponent*    Comp     = CastChecked<UGeometryCollectionComponent>(InComponent);
 		FGeometryCollectionPhysicsProxy* Proxy    = Comp->GetPhysicsProxy();
@@ -745,9 +737,7 @@ namespace Chaos
 
 	bool FGeometryCollectionCacheAdapter::InitializeForPlayback(UPrimitiveComponent* InComponent, UChaosCache* InCache, float InTime)
 	{
-#if PHYSICS_THREAD_CONTEXT
-		ensure(IsInGameThreadContext());
-#endif
+		EnsureIsInGameThreadContext();
 
 		if (UGeometryCollectionComponent* Comp = CastChecked<UGeometryCollectionComponent>(InComponent))
 		{
@@ -835,9 +825,7 @@ namespace Chaos
 
 	void FGeometryCollectionCacheAdapter::HandleBreakingEvents(const Chaos::FBreakingEventData& Event)
 	{
-#if PHYSICS_THREAD_CONTEXT
-		ensure(IsInGameThreadContext());
-#endif
+		EnsureIsInGameThreadContext();
 		
 		BreakingDataArray = &Event.BreakingData.AllBreakingsArray;
 
@@ -856,9 +844,7 @@ namespace Chaos
 
 	void FGeometryCollectionCacheAdapter::HandleCollisionEvents(const Chaos::FCollisionEventData& Event)
 	{
-#if PHYSICS_THREAD_CONTEXT
-		ensure(IsInGameThreadContext());
-#endif
+		EnsureIsInGameThreadContext();
 		
 		CollisionDataArray = &Event.CollisionData.AllCollisionsArray;
 
@@ -877,9 +863,7 @@ namespace Chaos
 
 	void FGeometryCollectionCacheAdapter::HandleTrailingEvents(const Chaos::FTrailingEventData& Event)
 	{
-#if PHYSICS_THREAD_CONTEXT
-		ensure(IsInGameThreadContext());
-#endif
+		EnsureIsInGameThreadContext();
 		
 		TrailingDataArray = &Event.TrailingData.AllTrailingsArray;
 
diff --git a/Engine/Source/Runtime/Experimental/Chaos/Private/EventDefaults.cpp b/Engine/Source/Runtime/Experimental/Chaos/Private/EventDefaults.cpp
--- a/Engine/Source/Runtime/Experimental/Chaos/Private/EventDefaults.cpp
+++ b/Engine/Source/Runtime/Experimental/Chaos/Private/EventDefaults.cpp
@@ -33,9 +33,7 @@ namespace Chaos
 		{
 			check(Solver);
 
-#if PHYSICS_THREAD_CONTEXT
-			ensure(IsInPhysicsThreadContext());
-#endif
+			EnsureIsInPhysicsThreadContext();
 
 			SCOPE_CYCLE_COUNTER(STAT_GatherCollisionEvent);
 
@@ -239,9 +237,7 @@ namespace Chaos
 		{
 			check(Solver);
 
-#if PHYSICS_THREAD_CONTEXT
-			ensure(IsInPhysicsThreadContext());
-#endif
+			EnsureIsInPhysicsThreadContext();
 
 			SCOPE_CYCLE_COUNTER(STAT_GatherBreakingEvent);
 
@@ -314,9 +310,7 @@ namespace Chaos
 		{
 			check(Solver);
 
-#if PHYSICS_THREAD_CONTEXT
-			ensure(IsInPhysicsThreadContext());
-#endif
+			EnsureIsInPhysicsThreadContext();
 
 			// #todo: This isn't working - SolverActor parameters are set on a solver but it is currently a different solver that is simulating!!
 			if (!Solver->GetEventFilters()->IsTrailingEventEnabled())
@@ -412,9 +406,7 @@ namespace Chaos
 		{
 			check(Solver);
 
-#if PHYSICS_THREAD_CONTEXT
-			ensure(IsInPhysicsThreadContext());
-#endif
+			EnsureIsInPhysicsThreadContext();
 
 			SCOPE_CYCLE_COUNTER(STAT_GatherSleepingEvent);
 
@@ -455,10 +447,7 @@ namespace Chaos
 		(const Chaos::FPBDRigidsSolver* Solver, FRemovalEventData& RemovalEventData)
 			{
 				check(Solver);
-#if PHYSICS_THREAD_CONTEXT
-				ensure(IsInPhysicsThreadContext());
-#endif
-
+				EnsureIsInPhysicsThreadContext();
 				
 				FRemovalDataArray& AllRemovalDataArray = RemovalEventData.RemovalData.AllRemovalArray;
 				TMap<IPhysicsProxyBase*, TArray<int32>>& AllRemovalIndicesByPhysicsProxy = RemovalEventData.PhysicsProxyToRemovalIndices.PhysicsProxyToIndicesMap;
diff --git a/Engine/Source/Runtime/Experimental/Chaos/Private/PBDRigidsSolver.cpp b/Engine/Source/Runtime/Experimental/Chaos/Private/PBDRigidsSolver.cpp
--- a/Engine/Source/Runtime/Experimental/Chaos/Private/PBDRigidsSolver.cpp
+++ b/Engine/Source/Runtime/Experimental/Chaos/Private/PBDRigidsSolver.cpp
@@ -643,9 +643,7 @@ namespace Chaos
 
 	void FPBDRigidsSolver::SetSuspensionTargetOnPhysicsThread(Chaos::FSuspensionConstraint* GTConstraint, const FVector& TargetPos, bool Enabled)
 	{
-#if PHYSICS_THREAD_CONTEXT
-		ensure(IsInPhysicsThreadContext());
-#endif
+		EnsureIsInPhysicsThreadContext();
 		FSuspensionConstraintPhysicsProxy* SuspensionProxy = GTConstraint->GetProxy<FSuspensionConstraintPhysicsProxy>();
 		check(SuspensionProxy);
 		SuspensionProxy->UpdateTargetOnPhysicsThread(this, TargetPos, Enabled);
diff --git a/Engine/Source/Runtime/Experimental/Chaos/Private/RewindData.cpp b/Engine/Source/Runtime/Experimental/Chaos/Private/RewindData.cpp
--- a/Engine/Source/Runtime/Experimental/Chaos/Private/RewindData.cpp
+++ b/Engine/Source/Runtime/Experimental/Chaos/Private/RewindData.cpp
@@ -129,10 +129,7 @@ bool FRewindData::RewindToFrame(int32 Frame)
 {
 	QUICK_SCOPE_CYCLE_COUNTER(RewindToFrame);
 
-#if PHYSICS_THREAD_CONTEXT
-	ensure(IsInPhysicsThreadContext());
-#endif
-
+	EnsureIsInPhysicsThreadContext();
 	//Can't go too far back
 	const int32 EarliestFrame = GetEarliestFrame_Internal();
 	if(Frame < EarliestFrame)
diff --git a/Engine/Source/Runtime/Experimental/Chaos/Public/Framework/Threading.h b/Engine/Source/Runtime/Experimental/Chaos/Public/Framework/Threading.h
--- a/Engine/Source/Runtime/Experimental/Chaos/Public/Framework/Threading.h
+++ b/Engine/Source/Runtime/Experimental/Chaos/Public/Framework/Threading.h
@@ -115,15 +115,23 @@ FORCEINLINE bool IsInGameThreadContext()
 {
 	return FPhysicsThreadContext::Get().IsInGameThreadContext();
 }
+
+FORCEINLINE void EnsureIsInPhysicsThreadContext()
+{
+	ensure(IsInPhysicsThreadContext());
+}
+
+FORCEINLINE void EnsureIsInGameThreadContext()
+{
+	ensure(IsInGameThreadContext());
+}
 #else
-FORCEINLINE bool IsInPhysicsThreadContext()
+FORCEINLINE void EnsureIsInPhysicsThreadContext()
 {
-	return false;
 }
 
-FORCEINLINE bool IsInGameThreadContext()
+FORCEINLINE void EnsureIsInGameThreadContext()
 {
-	return false;
 }
 #endif
 
diff --git a/Engine/Source/Runtime/Experimental/Chaos/Public/PhysicsProxy/SingleParticlePhysicsProxy.h b/Engine/Source/Runtime/Experimental/Chaos/Public/PhysicsProxy/SingleParticlePhysicsProxy.h
--- a/Engine/Source/Runtime/Experimental/Chaos/Public/PhysicsProxy/SingleParticlePhysicsProxy.h
+++ b/Engine/Source/Runtime/Experimental/Chaos/Public/PhysicsProxy/SingleParticlePhysicsProxy.h
@@ -800,12 +800,12 @@ protected:
 			//if proxy is registered with solver, we need a lock
 			if(GetSolverBase() != nullptr)
 			{
-				ensure(IsInGameThreadContext());
+				EnsureIsInGameThreadContext();
 			}
 		}
 		else
 		{
-			ensure(IsInPhysicsThreadContext());
+			EnsureIsInPhysicsThreadContext();
 		}
 #endif
 	}
