Fix uninitialized struct properties and non-unity issue

diff --git a/Engine/Plugins/Experimental/ChaosCaching/Source/ChaosCaching/Public/Chaos/ChaosCache.h b/Engine/Plugins/Experimental/ChaosCaching/Source/ChaosCaching/Public/Chaos/ChaosCache.h
--- a/Engine/Plugins/Experimental/ChaosCaching/Source/ChaosCaching/Public/Chaos/ChaosCache.h
+++ b/Engine/Plugins/Experimental/ChaosCaching/Source/ChaosCaching/Public/Chaos/ChaosCache.h
@@ -26,11 +26,11 @@ struct FParticleTransformTrack
 
 	/** The offset from the beginning of the cache that holds this track that the track starts */
 	UPROPERTY()
-	float BeginOffset;
+	float BeginOffset = 0.f;
 
 	/** If this flag is set true, the particle represented by this track deactives on the final keyframe. */
 	UPROPERTY()
-	bool bDeactivateOnEnd;
+	bool bDeactivateOnEnd = false;
 
 	/**
 	 * The above raw track is just the key data and doesn't know at which time those keys are placed, this is
