Fix cook warning on ChaosCachePlayer due to editor icon trying to be attached when not in editor

diff --git a/Engine/Plugins/Experimental/ChaosCaching/Source/ChaosCaching/Private/Chaos/CacheManagerActor.cpp b/Engine/Plugins/Experimental/ChaosCaching/Source/ChaosCaching/Private/Chaos/CacheManagerActor.cpp
--- a/Engine/Plugins/Experimental/ChaosCaching/Source/ChaosCaching/Private/Chaos/CacheManagerActor.cpp
+++ b/Engine/Plugins/Experimental/ChaosCaching/Source/ChaosCaching/Private/Chaos/CacheManagerActor.cpp
@@ -66,30 +66,33 @@ AChaosCacheManager::AChaosCacheManager(const FObjectInitializer& ObjectInitializ
 
 	// Add a sprite when in the editor
 #if WITH_EDITOR
-	struct FConstructorStatics
+	if (!IsRunningCommandlet())
 	{
-		ConstructorHelpers::FObjectFinderOptional<UTexture2D> SpriteTextureObject;
-		FName                                                 ID_CacheManager;
-		FText                                                 NAME_CacheManager;
-
-		FConstructorStatics()
-			: SpriteTextureObject(TEXT("/Engine/EditorResources/S_Actor"))
-			, ID_CacheManager(TEXT("Cache Manager"))
-			, NAME_CacheManager(NSLOCTEXT("SpriteCategory", "CacheManager", "Chaos Cache Manager"))
+		struct FConstructorStatics
 		{
-		}
-	};
-	static FConstructorStatics ConstructorStatics;
+			ConstructorHelpers::FObjectFinderOptional<UTexture2D> SpriteTextureObject;
+			FName                                                 ID_CacheManager;
+			FText                                                 NAME_CacheManager;
+
+			FConstructorStatics()
+				: SpriteTextureObject(TEXT("/Engine/EditorResources/S_Actor"))
+				, ID_CacheManager(TEXT("Cache Manager"))
+				, NAME_CacheManager(NSLOCTEXT("SpriteCategory", "CacheManager", "Chaos Cache Manager"))
+			{
+			}
+		};
+		static FConstructorStatics ConstructorStatics;
 
-	UBillboardComponent* SpriteComp = ObjectInitializer.CreateEditorOnlyDefaultSubobject<UBillboardComponent>(this, TEXT("Editor Icon"));
+		UBillboardComponent* SpriteComp = ObjectInitializer.CreateEditorOnlyDefaultSubobject<UBillboardComponent>(this, TEXT("Editor Icon"));
 
-	if(SpriteComp)
-	{
-		SpriteComp->Sprite                 = ConstructorStatics.SpriteTextureObject.Get();
-		SpriteComp->SpriteInfo.Category    = ConstructorStatics.ID_CacheManager;
-		SpriteComp->SpriteInfo.DisplayName = ConstructorStatics.NAME_CacheManager;
-		SpriteComp->Mobility               = EComponentMobility::Static;
-		SpriteComp->AttachToComponent(RootComponent, FAttachmentTransformRules::KeepRelativeTransform);
+		if (SpriteComp)
+		{
+			SpriteComp->Sprite = ConstructorStatics.SpriteTextureObject.Get();
+			SpriteComp->SpriteInfo.Category = ConstructorStatics.ID_CacheManager;
+			SpriteComp->SpriteInfo.DisplayName = ConstructorStatics.NAME_CacheManager;
+			SpriteComp->Mobility = EComponentMobility::Static;
+			SpriteComp->AttachToComponent(RootComponent, FAttachmentTransformRules::KeepRelativeTransform);
+		}
 	}
 #endif
 }
