Fixed "Section Not Found" warning on opening the context menu in the level editor

diff --git a/Engine/Plugins/Experimental/ChaosCaching/Source/ChaosCachingEditor/ChaosCachingEditor.Build.cs b/Engine/Plugins/Experimental/ChaosCaching/Source/ChaosCachingEditor/ChaosCachingEditor.Build.cs
--- a/Engine/Plugins/Experimental/ChaosCaching/Source/ChaosCachingEditor/ChaosCachingEditor.Build.cs
+++ b/Engine/Plugins/Experimental/ChaosCaching/Source/ChaosCachingEditor/ChaosCachingEditor.Build.cs
@@ -25,7 +25,8 @@ namespace UnrealBuildTool.Rules
 					"ToolMenus",
 					"PhysicsCore",
 					"ChaosCaching",
-					"GeometryCollectionEngine"
+					"GeometryCollectionEngine",
+					"LevelEditor"
 				});
 
 			SetupModulePhysicsSupport(Target);
diff --git a/Engine/Plugins/Experimental/ChaosCaching/Source/ChaosCachingEditor/Private/Chaos/ChaosCachingEditorPlugin.cpp b/Engine/Plugins/Experimental/ChaosCaching/Source/ChaosCachingEditor/Private/Chaos/ChaosCachingEditorPlugin.cpp
--- a/Engine/Plugins/Experimental/ChaosCaching/Source/ChaosCachingEditor/Private/Chaos/ChaosCachingEditorPlugin.cpp
+++ b/Engine/Plugins/Experimental/ChaosCaching/Source/ChaosCachingEditor/Private/Chaos/ChaosCachingEditorPlugin.cpp
@@ -17,6 +17,7 @@
 #include "Modules/ModuleManager.h"
 #include "PropertyEditorModule.h"
 #include "ToolMenus.h"
+#include "LevelEditor.h"
 
 IMPLEMENT_MODULE(IChaosCachingEditorPlugin, ChaosCachingEditor)
 
@@ -32,7 +33,12 @@ void IChaosCachingEditorPlugin::StartupModule()
 
 	FCachingEditorCommands::Register();
 
-	StartupHandle = UToolMenus::RegisterStartupCallback(FSimpleMulticastDelegate::FDelegate::CreateRaw(this, &IChaosCachingEditorPlugin::RegisterMenus));
+	// Register level editor menu extender
+	FLevelEditorModule::FLevelViewportMenuExtender_SelectedActors LevelEditorMenuExtenderDelegate = FLevelEditorModule::FLevelViewportMenuExtender_SelectedActors::CreateRaw(this, &IChaosCachingEditorPlugin::ExtendLevelViewportContextMenu);
+	FLevelEditorModule& LevelEditorModule = FModuleManager::LoadModuleChecked<FLevelEditorModule>(TEXT("LevelEditor"));
+	auto& MenuExtenders = LevelEditorModule.GetAllLevelViewportContextMenuExtenders();
+	MenuExtenders.Add(LevelEditorMenuExtenderDelegate);
+	StartupHandle = MenuExtenders.Last().GetHandle();
 
 	FPropertyEditorModule& PropertyModule = FModuleManager::GetModuleChecked<FPropertyEditorModule>("PropertyEditor");
 	PropertyModule.RegisterCustomClassLayout("ChaosCacheCollection", FOnGetDetailCustomizationInstance::CreateStatic(&FCacheCollectionDetails::MakeInstance));
@@ -49,7 +55,14 @@ void IChaosCachingEditorPlugin::ShutdownModule()
 		PropertyModule.UnregisterCustomClassLayout("ChaosCacheManager");
 		PropertyModule.UnregisterCustomClassLayout("ChaosCacheCollection");
 
-		UToolMenus::UnRegisterStartupCallback(StartupHandle);
+		// Unregister level editor menu extender
+		if (FModuleManager::Get().IsModuleLoaded(TEXT("LevelEditor")))
+		{
+			FLevelEditorModule& LevelEditorModule = FModuleManager::GetModuleChecked<FLevelEditorModule>(TEXT("LevelEditor"));
+			LevelEditorModule.GetAllLevelViewportContextMenuExtenders().RemoveAll([&](const FLevelEditorModule::FLevelViewportMenuExtender_SelectedActors& Delegate) {
+				return Delegate.GetHandle() == StartupHandle;
+				});
+		}
 
 		FCachingEditorCommands::Unregister();
 
@@ -64,39 +77,39 @@ bool IsCreateCacheManagerVisible();
 bool IsSetAllRecordVisible();
 bool IsSetAllPlayVisible();
 
-void IChaosCachingEditorPlugin::RegisterMenus()
+TSharedRef<FExtender> IChaosCachingEditorPlugin::ExtendLevelViewportContextMenu(const TSharedRef<FUICommandList> InCommandList, const TArray<AActor*> SelectedActors)
 {
-	FToolMenuOwnerScoped OwnerScope(this);
+	TSharedRef<FExtender> Extender(new FExtender());
 
-	UToolMenu* Menu = UToolMenus::Get()->ExtendMenu("LevelEditor.ActorContextMenu");
-
-	FToolMenuSection* Section = Menu->FindSection("Chaos");
-	if(!Section)
-	{
-		Section = &Menu->AddSection("Chaos");
-	}
+	Extender->AddMenuExtension("ActorUETools", EExtensionHook::After, InCommandList, FMenuExtensionDelegate::CreateLambda(
+		[this, SelectedActors](FMenuBuilder& MenuBuilder)
+		{
+			MenuBuilder.AddSubMenu(
+				LOCTEXT("ChaosSectionLabel", "Chaos"),
+				LOCTEXT("Tooltip_Caching", "Options for manipulating cache managers and their observed components"),
+				FNewMenuDelegate::CreateLambda([this](FMenuBuilder &InMenuBuilder) {
+					RegisterCachingSubMenu(InMenuBuilder);
+					}),
+				FUIAction(FExecuteAction(), FCanExecuteAction::CreateLambda([]() -> bool
+					{
+						return IsCreateCacheManagerVisible() || IsSetAllPlayVisible() || IsSetAllRecordVisible();
+					})),
+					NAME_None,
+					EUserInterfaceActionType::Button,
+					false,
+					FSlateIcon(FChaosCachingEditorStyle::Get().GetStyleSetName(), "ChaosCachingEditor.Fracture"));		
+		}
+	));
 
-	Section->InitSection("Chaos", TAttribute<FText>(), FToolMenuInsert("ActorUETools", EToolMenuInsertType::After));
-
-	Section->AddSubMenu("CachingSub",
-						LOCTEXT("ChaosSectionLabel", "Chaos"),
-						LOCTEXT("Tooltip_Caching", "Options for manipulating cache managers and their observed components"),
-						FNewToolMenuDelegate::CreateLambda([this](UToolMenu* InMenu) {
-							FToolMenuSection& CacheSubMenuSection = InMenu->AddSection("Caching", LOCTEXT("SubMenu_Caching", "Caching"));
-							RegisterCachingSubMenu(InMenu, &CacheSubMenuSection);
-						}),
-						FUIAction(FExecuteAction(), FCanExecuteAction::CreateLambda([]() -> bool 
-						{
-							return IsCreateCacheManagerVisible() || IsSetAllPlayVisible() || IsSetAllRecordVisible();
-						})),
-						EUserInterfaceActionType::Button,
-						false,
-						FSlateIcon(FChaosCachingEditorStyle::Get().GetStyleSetName(), "ChaosCachingEditor.Fracture"));
+	return Extender;
 }
 
-void IChaosCachingEditorPlugin::RegisterCachingSubMenu(UToolMenu* InMenu, FToolMenuSection* InSection)
+
+void IChaosCachingEditorPlugin::RegisterCachingSubMenu(FMenuBuilder& InMenuBuilder)
 {
-	InSection->AddMenuEntry("CreateCacheManager",
+	InMenuBuilder.BeginSection("Caching", LOCTEXT("SubMenu_Caching", "Caching"));
+
+	InMenuBuilder.AddMenuEntry(
 							LOCTEXT("MenuItem_CreateCacheManager", "Create Cache Manager"),
 							LOCTEXT("MenuItem_CreateCacheManager_ToolTip", "Adds a cache manager to observe compatible components in the selection set."),
 							FSlateIcon(),
@@ -104,6 +117,7 @@ void IChaosCachingEditorPlugin::RegisterCachingSubMenu(UToolMenu* InMenu, FToolM
 									  FCanExecuteAction(),
 									  FIsActionChecked(),
 									  FIsActionButtonVisible::CreateStatic(&IsCreateCacheManagerVisible)));
+	InMenuBuilder.EndSection();
 }
 
 void IChaosCachingEditorPlugin::OnCreateCacheManager()
diff --git a/Engine/Plugins/Experimental/ChaosCaching/Source/ChaosCachingEditor/Public/Chaos/ChaosCachingEditorPlugin.h b/Engine/Plugins/Experimental/ChaosCaching/Source/ChaosCachingEditor/Public/Chaos/ChaosCachingEditorPlugin.h
--- a/Engine/Plugins/Experimental/ChaosCaching/Source/ChaosCachingEditor/Public/Chaos/ChaosCachingEditorPlugin.h
+++ b/Engine/Plugins/Experimental/ChaosCaching/Source/ChaosCachingEditor/Public/Chaos/ChaosCachingEditorPlugin.h
@@ -6,11 +6,13 @@
 #include "Modules/ModuleInterface.h"
 #include "Modules/ModuleManager.h"
 #include "Delegates/IDelegateInstance.h"
+#include "Framework/MultiBox/MultiBoxExtender.h"
 
 class FAssetTypeActions_ChaosCacheCollection;
 
 class UToolMenu;
 struct FToolMenuSection;
+class AActor;
 
 /**
  * The public interface to this module
@@ -43,8 +45,8 @@ public:
 
 private:
 
-	void RegisterMenus();
-	void RegisterCachingSubMenu(UToolMenu* InMenu, FToolMenuSection* InSection);
+	TSharedRef<FExtender> ExtendLevelViewportContextMenu(const TSharedRef<FUICommandList> InCommandList, const TArray<AActor*> SelectedActors);
+	void RegisterCachingSubMenu(FMenuBuilder& InMenuBuilder);
 	
 	void OnCreateCacheManager();
 
