Chaos : Caching
- don’t insert empty elements into caching array.

diff --git a/Engine/Plugins/Experimental/ChaosCaching/Source/ChaosCaching/Private/Chaos/CacheCollection.cpp b/Engine/Plugins/Experimental/ChaosCaching/Source/ChaosCaching/Private/Chaos/CacheCollection.cpp
--- a/Engine/Plugins/Experimental/ChaosCaching/Source/ChaosCaching/Private/Chaos/CacheCollection.cpp
+++ b/Engine/Plugins/Experimental/ChaosCaching/Source/ChaosCaching/Private/Chaos/CacheCollection.cpp
@@ -15,8 +15,10 @@ void UChaosCacheCollection::GetAssetRegistryTags(TArray<FAssetRegistryTag>& OutT
 	float Duration = 0.0;
 	for (int32 CacheIdx = 0; CacheIdx < NumCaches; ++CacheIdx)
 	{
-		UChaosCache* Cache = Caches[CacheIdx];
-		Duration = FMath::Max(Duration, Cache->RecordedDuration);
+		if (UChaosCache* Cache = Caches[CacheIdx])
+		{
+			Duration = FMath::Max(Duration, Cache->RecordedDuration);
+		}
 	}	
 	OutTags.Add(FAssetRegistryTag(TEXT("Recorded Duration"), FString::Printf(TEXT("%.2f"), Duration), FAssetRegistryTag::TT_Numerical));
 #endif
@@ -28,7 +30,11 @@ UChaosCache* UChaosCacheCollection::FindCache(const FName& CacheName) const
 {
 	UChaosCache* const* ExistingCache = Algo::FindByPredicate(Caches, [&CacheName](const UChaosCache* Test)
 	{
-		return Test->GetFName() == CacheName;
+		if (Test)
+		{
+			return Test->GetFName() == CacheName;
+		}
+		return false;
 	});
 
 	return ExistingCache ? *ExistingCache : nullptr;
@@ -43,7 +49,11 @@ UChaosCache* UChaosCacheCollection::FindOrAddCache(const FName& CacheName)
 	{
 		UChaosCache** ExistingCache = Algo::FindByPredicate(Caches, [&FinalName](const UChaosCache* Test)
 		{
-			return Test->GetFName() == FinalName;
+			if (Test)
+			{
+				return Test->GetFName() == FinalName;
+			}
+			return false;
 		});
 
 		ResultCache = ExistingCache ? *ExistingCache : nullptr;
@@ -70,7 +80,10 @@ void UChaosCacheCollection::FlushAllCacheWrites()
 {
 	ParallelFor(Caches.Num(), [this](int32 InIndex)
 	{
-		Caches[InIndex]->FlushPendingFrames();
+		if (Caches[InIndex])
+		{
+			Caches[InIndex]->FlushPendingFrames();
+		}
 	});
 }
 
diff --git a/Engine/Plugins/Experimental/ChaosCaching/Source/ChaosCachingEditor/Private/Chaos/CacheCollectionCustomization.cpp b/Engine/Plugins/Experimental/ChaosCaching/Source/ChaosCachingEditor/Private/Chaos/CacheCollectionCustomization.cpp
--- a/Engine/Plugins/Experimental/ChaosCaching/Source/ChaosCachingEditor/Private/Chaos/CacheCollectionCustomization.cpp
+++ b/Engine/Plugins/Experimental/ChaosCaching/Source/ChaosCachingEditor/Private/Chaos/CacheCollectionCustomization.cpp
@@ -32,16 +32,19 @@ void FCacheCollectionDetails::CustomizeDetails(IDetailLayoutBuilder& DetailBuild
 
 	Item = Objects[0];
 
-	if(UChaosCacheCollection* Collection = GetSelectedCollection())
+	if (UChaosCacheCollection* Collection = GetSelectedCollection())
 	{
-		NameEditBoxes.SetNum(Collection->Caches.Num());
+		if (Collection->Caches.Num())
+		{
+			NameEditBoxes.SetNum(Collection->Caches.Num());
 
-		TSharedRef<IPropertyHandle> CachesProp    = DetailBuilder.GetProperty("Caches");
-		IDetailCategoryBuilder&     CacheCategory = DetailBuilder.EditCategory(CachesProp->GetDefaultCategoryName());
+			TSharedRef<IPropertyHandle> CachesProp = DetailBuilder.GetProperty("Caches");
+			IDetailCategoryBuilder& CacheCategory = DetailBuilder.EditCategory(CachesProp->GetDefaultCategoryName());
 
-		TSharedRef<FDetailArrayBuilder> CacheArrayBuilder = MakeShareable(new FDetailArrayBuilder(CachesProp, true, false, true));
-		CacheArrayBuilder->OnGenerateArrayElementWidget(FOnGenerateArrayElementWidget::CreateSP(this, &FCacheCollectionDetails::GenerateCacheArrayElementWidget, &DetailBuilder));
-		CacheCategory.AddCustomBuilder(CacheArrayBuilder);
+			TSharedRef<FDetailArrayBuilder> CacheArrayBuilder = MakeShareable(new FDetailArrayBuilder(CachesProp, true, false, true));
+			CacheArrayBuilder->OnGenerateArrayElementWidget(FOnGenerateArrayElementWidget::CreateSP(this, &FCacheCollectionDetails::GenerateCacheArrayElementWidget, &DetailBuilder));
+			CacheCategory.AddCustomBuilder(CacheArrayBuilder);
+		}
 	}
 }
 
@@ -158,6 +161,7 @@ void FCacheCollectionDetails::GenerateCacheArrayElementWidget(TSharedRef<IProper
 			InPropertyHandle->CreatePropertyNameWidget()
 		];
 
+
 	WidgetRow.ValueContent()
 		[
 			SNew(SHorizontalBox)
@@ -171,6 +175,7 @@ void FCacheCollectionDetails::GenerateCacheArrayElementWidget(TSharedRef<IProper
 				.OnTextChanged(this, &FCacheCollectionDetails::OnChangeCacheName, ArrayIndex)
 				.OnTextCommitted(this, &FCacheCollectionDetails::OnCommitCacheName, ArrayIndex)
 			]
+
 			+SHorizontalBox::Slot()
 			.VAlign(VAlign_Center)
 			.HAlign(HAlign_Fill)
