Chaos : Caching
- Fixed crash on cache asset manipulation.

diff --git a/Engine/Plugins/Experimental/ChaosCaching/Source/ChaosCachingEditor/Private/Chaos/CacheCollectionCustomization.cpp b/Engine/Plugins/Experimental/ChaosCaching/Source/ChaosCachingEditor/Private/Chaos/CacheCollectionCustomization.cpp
--- a/Engine/Plugins/Experimental/ChaosCaching/Source/ChaosCachingEditor/Private/Chaos/CacheCollectionCustomization.cpp
+++ b/Engine/Plugins/Experimental/ChaosCaching/Source/ChaosCachingEditor/Private/Chaos/CacheCollectionCustomization.cpp
@@ -24,6 +24,8 @@ void FCacheCollectionDetails::CustomizeDetails(IDetailLayoutBuilder& DetailBuild
 	TArray<TWeakObjectPtr<UObject>> Objects;
 	DetailBuilder.GetObjectsBeingCustomized(Objects);
 
+	NameEditBoxes.Empty();
+
 	// In the simple asset editor we should only get one object
 	if(Objects.Num() != 1)
 	{
@@ -36,7 +38,10 @@ void FCacheCollectionDetails::CustomizeDetails(IDetailLayoutBuilder& DetailBuild
 	{
 		if (Collection->Caches.Num())
 		{
-			NameEditBoxes.SetNum(Collection->Caches.Num());
+			for (int i=0;i<Collection->Caches.Num();i++)
+			{
+				NameEditBoxes.Add(i, TSharedPtr<SEditableTextBox>());
+			}
 
 			TSharedRef<IPropertyHandle> CachesProp = DetailBuilder.GetProperty("Caches");
 			IDetailCategoryBuilder& CacheCategory = DetailBuilder.EditCategory(CachesProp->GetDefaultCategoryName());
@@ -48,6 +53,19 @@ void FCacheCollectionDetails::CustomizeDetails(IDetailLayoutBuilder& DetailBuild
 	}
 }
 
+const UChaosCache* FCacheCollectionDetails::GetCacheCollection(int32 Index) const
+{
+	if (const UChaosCacheCollection* Collection = GetSelectedCollection())
+	{
+		if(0 <= Index && Index < Collection->Caches.Num())
+		{
+			return Collection->Caches[Index];
+		}
+	}
+	return nullptr;
+}
+
+
 UChaosCacheCollection* FCacheCollectionDetails::GetSelectedCollection()
 {
 	if(UObject* ItemObj = Item.Get())
@@ -91,7 +109,7 @@ void FCacheCollectionDetails::OnDeleteCache(int32 InArrayIndex, IDetailLayoutBui
 {
 	if(UChaosCacheCollection* Collection = GetSelectedCollection())
 	{
-		if(Collection->Caches.IsValidIndex(InArrayIndex) && Collection->Caches[InArrayIndex])
+		if(Collection->Caches.IsValidIndex(InArrayIndex) )
 		{
 			Collection->Caches.RemoveAt(InArrayIndex);
 			InLayoutBuilder->ForceRefreshDetails();
@@ -161,6 +179,11 @@ void FCacheCollectionDetails::GenerateCacheArrayElementWidget(TSharedRef<IProper
 			InPropertyHandle->CreatePropertyNameWidget()
 		];
 
+	TSharedPtr<SEditableTextBox> TextBox = SNew(SEditableTextBox);
+	if (NameEditBoxes.Contains(ArrayIndex))
+	{
+		TextBox = NameEditBoxes[ArrayIndex];
+	}
 
 	WidgetRow.ValueContent()
 		[
@@ -170,7 +193,7 @@ void FCacheCollectionDetails::GenerateCacheArrayElementWidget(TSharedRef<IProper
 			.HAlign(HAlign_Fill)
 			.AutoWidth()
 			[
-				SAssignNew(NameEditBoxes[ArrayIndex], SEditableTextBox)
+				SAssignNew(TextBox, SEditableTextBox)
 				.Text(this, &FCacheCollectionDetails::GetCacheName, ArrayIndex)
 				.OnTextChanged(this, &FCacheCollectionDetails::OnChangeCacheName, ArrayIndex)
 				.OnTextCommitted(this, &FCacheCollectionDetails::OnCommitCacheName, ArrayIndex)
diff --git a/Engine/Plugins/Experimental/ChaosCaching/Source/ChaosCachingEditor/Public/Chaos/CacheCollectionCustomization.h b/Engine/Plugins/Experimental/ChaosCaching/Source/ChaosCachingEditor/Public/Chaos/CacheCollectionCustomization.h
--- a/Engine/Plugins/Experimental/ChaosCaching/Source/ChaosCachingEditor/Public/Chaos/CacheCollectionCustomization.h
+++ b/Engine/Plugins/Experimental/ChaosCaching/Source/ChaosCachingEditor/Public/Chaos/CacheCollectionCustomization.h
@@ -13,6 +13,7 @@ class IDetailChildrenBuilder;
 
 class UObject;
 class UChaosCacheCollection;
+class UChaosCache;
 
 class SEditableTextBox;
 
@@ -26,6 +27,7 @@ private:
 
 	UChaosCacheCollection* GetSelectedCollection();
 	const UChaosCacheCollection* GetSelectedCollection() const;
+	const UChaosCache* GetCacheCollection(int32 Index) const;
 	FText GetCacheName(int32 InCacheIndex) const;
 	void OnDeleteCache(int32 InArrayIndex, IDetailLayoutBuilder* InLayoutBuilder);
 
@@ -35,5 +37,5 @@ private:
 	void GenerateCacheArrayElementWidget(TSharedRef<IPropertyHandle> PropertyHandle, int32 ArrayIndex, IDetailChildrenBuilder& ChildrenBuilder, IDetailLayoutBuilder* DetailLayout);
 	TWeakObjectPtr<UObject> Item;
 
-	TArray<TSharedPtr<SEditableTextBox>> NameEditBoxes;
+	TMap< int, TSharedPtr<SEditableTextBox>> NameEditBoxes;
 };
\ No newline at end of file
