Chaos : Caching
- Update world space state and dirty elements when caching updates the particle position.

diff --git a/Engine/Plugins/Experimental/ChaosCaching/Source/ChaosCaching/Private/Chaos/Adapters/GeometryCollectionComponentCacheAdapter.cpp b/Engine/Plugins/Experimental/ChaosCaching/Source/ChaosCaching/Private/Chaos/Adapters/GeometryCollectionComponentCacheAdapter.cpp
--- a/Engine/Plugins/Experimental/ChaosCaching/Source/ChaosCaching/Private/Chaos/Adapters/GeometryCollectionComponentCacheAdapter.cpp
+++ b/Engine/Plugins/Experimental/ChaosCaching/Source/ChaosCaching/Private/Chaos/Adapters/GeometryCollectionComponentCacheAdapter.cpp
@@ -507,6 +507,9 @@ namespace Chaos
 				Handle->SetX(Handle->P());
 				Handle->SetR(Handle->Q());
 				
+				Handle->UpdateWorldSpaceState(WorldTransform, FVec3(0));
+				Solver->GetEvolution()->DirtyParticle(*Handle);
+
 				if(FRigidParticle* ClusterParent = Handle->ClusterIds().Id)
 				{
 					if(FClusterParticle* Parent = ClusterParent->CastToClustered())
@@ -526,6 +529,9 @@ namespace Chaos
 							Parent->SetX(Result.GetTranslation());
 							Parent->SetQ(Result.GetRotation());
 							Parent->SetR(Result.GetRotation());
+
+							Parent->UpdateWorldSpaceState(WorldTransform, FVec3(0));
+							Solver->GetEvolution()->DirtyParticle(*Parent);
 						}
 					}
 				}
