Fixed vaulting not working properly over cached geometry collection in AncientGame
- Made sure disabled particles are not marked as dirty by the geometry coillection cache code

diff --git a/Engine/Plugins/Experimental/ChaosCaching/Source/ChaosCaching/Private/Chaos/Adapters/GeometryCollectionComponentCacheAdapter.cpp b/Engine/Plugins/Experimental/ChaosCaching/Source/ChaosCaching/Private/Chaos/Adapters/GeometryCollectionComponentCacheAdapter.cpp
--- a/Engine/Plugins/Experimental/ChaosCaching/Source/ChaosCaching/Private/Chaos/Adapters/GeometryCollectionComponentCacheAdapter.cpp
+++ b/Engine/Plugins/Experimental/ChaosCaching/Source/ChaosCaching/Private/Chaos/Adapters/GeometryCollectionComponentCacheAdapter.cpp
@@ -508,7 +508,11 @@ namespace Chaos
 				Handle->SetR(Handle->Q());
 				
 				Handle->UpdateWorldSpaceState(WorldTransform, FVec3(0));
-				Solver->GetEvolution()->DirtyParticle(*Handle);
+
+				if(!Handle->Disabled())
+				{
+					Solver->GetEvolution()->DirtyParticle(*Handle);
+				}
 
 				if(FRigidParticle* ClusterParent = Handle->ClusterIds().Id)
 				{
@@ -530,8 +534,12 @@ namespace Chaos
 							Parent->SetQ(Result.GetRotation());
 							Parent->SetR(Result.GetRotation());
 
-							Parent->UpdateWorldSpaceState(WorldTransform, FVec3(0));
-							Solver->GetEvolution()->DirtyParticle(*Parent);
+							Parent->UpdateWorldSpaceState(Result, FVec3(0));
+
+							if(!Parent->Disabled())
+							{
+								Solver->GetEvolution()->DirtyParticle(*Parent);
+							}
 						}
 					}
 				}
