Chaos : Fix cache manager not working after when reseting a geometry collection
- Make sure EndPlayback is called even if the ValidForPlayback for playback fail so that the playback count is properly decremented

diff --git a/Engine/Plugins/Experimental/ChaosCaching/Source/ChaosCaching/Private/Chaos/CacheManagerActor.cpp b/Engine/Plugins/Experimental/ChaosCaching/Source/ChaosCaching/Private/Chaos/CacheManagerActor.cpp
--- a/Engine/Plugins/Experimental/ChaosCaching/Source/ChaosCaching/Private/Chaos/CacheManagerActor.cpp
+++ b/Engine/Plugins/Experimental/ChaosCaching/Source/ChaosCaching/Private/Chaos/CacheManagerActor.cpp
@@ -567,9 +567,12 @@ void AChaosCacheManager::OnStartFrameChanged(Chaos::FReal InTime)
 			}
 
 			FCacheUserToken Token = Observed.Cache->BeginPlayback();
-			if (Token.IsOpen() && Observed.BestFitAdapter->ValidForPlayback(Comp, Observed.Cache))
+			if (Token.IsOpen())
 			{
-				Observed.BestFitAdapter->SetRestState(Comp, Observed.Cache, GetTransform(), InTime);
+				if (Observed.BestFitAdapter->ValidForPlayback(Comp, Observed.Cache))
+				{
+					Observed.BestFitAdapter->SetRestState(Comp, Observed.Cache, GetTransform(), InTime);
+				}
 				Observed.Cache->EndPlayback(Token);
 			}
 		}
