Chaos geometry collection : Fix off center rotation issues caused by interpolation ( async and cache )
- properly multiply by MassToLocal before interpolating fixes the issue

diff --git a/Engine/Plugins/Experimental/ChaosCaching/Source/ChaosCaching/Private/Chaos/Adapters/GeometryCollectionComponentCacheAdapter.cpp b/Engine/Plugins/Experimental/ChaosCaching/Source/ChaosCaching/Private/Chaos/Adapters/GeometryCollectionComponentCacheAdapter.cpp
--- a/Engine/Plugins/Experimental/ChaosCaching/Source/ChaosCaching/Private/Chaos/Adapters/GeometryCollectionComponentCacheAdapter.cpp
+++ b/Engine/Plugins/Experimental/ChaosCaching/Source/ChaosCaching/Private/Chaos/Adapters/GeometryCollectionComponentCacheAdapter.cpp
@@ -267,7 +267,7 @@ namespace Chaos
 		Context.bEvaluateCurves    = false;
 		Context.bEvaluateEvents    = true;
 
-		FCacheEvaluationResult EvaluatedResult = InCache->Evaluate(Context);
+		FCacheEvaluationResult EvaluatedResult = InCache->Evaluate(Context, &MassToLocal.GetConstArray());
 
 		const int32                      NumEventTracks = EvaluatedResult.Events.Num();
 		const TArray<FCacheEventHandle>* EnableEvents   = EvaluatedResult.Events.Find(FEnableStateEvent::EventName);
@@ -499,9 +499,10 @@ namespace Chaos
 					continue;
 				}
 
-				// Evaluated transform is in Cache Manager space with MassToLocal removed. We need it in World space.
-				FTransform WorldTransform = MassToLocal[ParticleIndex] * EvaluatedTransform;
-				
+				// EvaluatedTransform is in world space already as we have passed MassToLocal transform to InCache->Evaluate()
+				// to make sure the interpolation was done correctly in world space  
+				const FTransform WorldTransform{EvaluatedTransform};
+
 				Handle->SetP(WorldTransform.GetTranslation());
 				Handle->SetQ(WorldTransform.GetRotation());
 				Handle->SetX(Handle->P());
@@ -643,7 +644,8 @@ namespace Chaos
 				Context.bEvaluateCurves = false;
 				Context.bEvaluateEvents = false;
 
-				FCacheEvaluationResult EvaluatedResult = InCache->Evaluate(Context);
+				// we only need the rest transforms, so no need to pass MassToLocal transforms 
+				FCacheEvaluationResult EvaluatedResult = InCache->Evaluate(Context, nullptr);
 				const int32 NumCacheTransforms = EvaluatedResult.Transform.Num();
 			
 				// Any bone that is not explicitly set by the cache defaults to it's rest collection position.
@@ -765,7 +767,8 @@ namespace Chaos
 					Context.bEvaluateCurves = false;
 					Context.bEvaluateEvents = false;
 
-					FCacheEvaluationResult EvaluatedResult = InCache->Evaluate(Context);
+					// initial transforms do no need to be multiplied by MassToLocal transforms so we can pass nullptr
+					FCacheEvaluationResult EvaluatedResult = InCache->Evaluate(Context, nullptr);
 					const int32 NumCacheTransforms = EvaluatedResult.Transform.Num();
 
 					// Any bone that is not explicitly set by the cache defaults to it's rest collection position.
@@ -814,7 +817,8 @@ namespace Chaos
 		Context.bEvaluateCurves = false;
 		Context.bEvaluateEvents = true;
 
-		FCacheEvaluationResult EvaluatedResult = InCache->Evaluate(Context);
+		// passing nullptr for the MassToLocal as we do not need to evaluate the transforms
+		FCacheEvaluationResult EvaluatedResult = InCache->Evaluate(Context, nullptr);
 		const TArray<FCacheEventHandle>* EnableEvents = EvaluatedResult.Events.Find(FEnableStateEvent::EventName);
 		if (EnableEvents)
 		{
diff --git a/Engine/Plugins/Experimental/ChaosCaching/Source/ChaosCaching/Private/Chaos/Adapters/StaticMeshComponentCacheAdapter.cpp b/Engine/Plugins/Experimental/ChaosCaching/Source/ChaosCaching/Private/Chaos/Adapters/StaticMeshComponentCacheAdapter.cpp
--- a/Engine/Plugins/Experimental/ChaosCaching/Source/ChaosCaching/Private/Chaos/Adapters/StaticMeshComponentCacheAdapter.cpp
+++ b/Engine/Plugins/Experimental/ChaosCaching/Source/ChaosCaching/Private/Chaos/Adapters/StaticMeshComponentCacheAdapter.cpp
@@ -84,7 +84,8 @@ namespace Chaos
 				Context.bEvaluateCurves = false;
 				Context.bEvaluateEvents = false;
 
-				FCacheEvaluationResult EvaluatedResult = InCache->Evaluate(Context);
+				// Note: static mesh don't use the concept of MassToLocal transforms, passing nullptr
+				FCacheEvaluationResult EvaluatedResult = InCache->Evaluate(Context, nullptr);
 
 				// Either 0 or 1 result, 0 for nothing in the eval track - 1 if there was.
 				if(EvaluatedResult.Transform.Num() == 1)
@@ -154,7 +155,7 @@ namespace Chaos
 			Context.bEvaluateEvents = false;
 
 			FTransform RestTransform;
-			InCache->EvaluateSingle(0, TickRecord, &RestTransform, nullptr);
+			InCache->EvaluateSingle(0, TickRecord, nullptr, &RestTransform, nullptr);
 
 			// Evaluated transform is in CacheManager space.
 			InComponent->SetWorldTransform(RestTransform * InRootTransform, false);
diff --git a/Engine/Plugins/Experimental/ChaosCaching/Source/ChaosCaching/Private/Chaos/ChaosCache.cpp b/Engine/Plugins/Experimental/ChaosCaching/Source/ChaosCaching/Private/Chaos/ChaosCache.cpp
--- a/Engine/Plugins/Experimental/ChaosCaching/Source/ChaosCaching/Private/Chaos/ChaosCache.cpp
+++ b/Engine/Plugins/Experimental/ChaosCaching/Source/ChaosCaching/Private/Chaos/ChaosCache.cpp
@@ -312,7 +312,7 @@ float UChaosCache::GetDuration() const
 	return RecordedDuration;
 }
 
-FCacheEvaluationResult UChaosCache::Evaluate(const FCacheEvaluationContext& InContext)
+FCacheEvaluationResult UChaosCache::Evaluate(const FCacheEvaluationContext& InContext, const TArray<FTransform>* MassToLocalTransforms)
 {
 	QUICK_SCOPE_CYCLE_COUNTER(QSTAT_CacheEval);
 
@@ -374,7 +374,8 @@ FCacheEvaluationResult UChaosCache::Evaluate(const FCacheEvaluationContext& InCo
 					Result.ParticleIndices.Add(INDEX_NONE);
 				}
 
-				EvaluateSingle(CacheIndex, InContext.TickRecord, EvalTransform, EvalCurves);
+				const FTransform* MassToLocal = MassToLocalTransforms? &((*MassToLocalTransforms)[CacheIndex]): nullptr;
+				EvaluateSingle(CacheIndex, InContext.TickRecord, MassToLocal, EvalTransform, EvalCurves);
 			}
 		}
 	}
@@ -430,7 +431,8 @@ FCacheEvaluationResult UChaosCache::Evaluate(const FCacheEvaluationContext& InCo
 				Result.ParticleIndices.Add(INDEX_NONE);
 			}
 
-			EvaluateSingle(Index, InContext.TickRecord, EvalTransform, EvalCurves);
+			const FTransform* MassToLocal = MassToLocalTransforms? &((*MassToLocalTransforms)[Index]): nullptr;
+			EvaluateSingle(Index, InContext.TickRecord, MassToLocal, EvalTransform, EvalCurves);
 		}
 	}
 
@@ -459,7 +461,7 @@ const FCacheSpawnableTemplate& UChaosCache::GetSpawnableTemplate() const
 	return Spawnable;
 }
 
-void UChaosCache::EvaluateSingle(int32 InIndex, FPlaybackTickRecord& InTickRecord, FTransform* OutOptTransform, TMap<FName, float>* OutOptCurves)
+void UChaosCache::EvaluateSingle(int32 InIndex, FPlaybackTickRecord& InTickRecord, const FTransform* MassToLocal, FTransform* OutOptTransform, TMap<FName, float>* OutOptCurves)
 {
 	// check to satisfy SA, external callers check validity in Evaluate
 	checkSlow(ParticleTracks.IsValidIndex(InIndex));
@@ -468,7 +470,7 @@ void UChaosCache::EvaluateSingle(int32 InIndex, FPlaybackTickRecord& InTickRecor
 	
 	if(OutOptTransform)
 	{
-		EvaluateTransform(Data, InTickRecord.GetTime(), *OutOptTransform);
+		EvaluateTransform(Data, InTickRecord.GetTime(), MassToLocal, *OutOptTransform);
 		(*OutOptTransform) = (*OutOptTransform) *InTickRecord.SpaceTransform;
 	}
 	
@@ -479,9 +481,9 @@ void UChaosCache::EvaluateSingle(int32 InIndex, FPlaybackTickRecord& InTickRecor
 	}
 }
 
-void UChaosCache::EvaluateTransform(const FPerParticleCacheData& InData, float InTime, FTransform& OutTransform)
+void UChaosCache::EvaluateTransform(const FPerParticleCacheData& InData, float InTime, const FTransform* MassToLocal, FTransform& OutTransform)
 {
-	OutTransform = InData.TransformData.Evaluate(InTime);
+	OutTransform = InData.TransformData.Evaluate(InTime, MassToLocal);
 }
 
 void UChaosCache::EvaluateCurves(const FPerParticleCacheData& InData, float InTime, TMap<FName, float>& OutCurves)
@@ -534,63 +536,51 @@ void UChaosCache::EvaluateEvents(FPlaybackTickRecord& InTickRecord, TMap<FName,
 	}
 }
 
-FTransform FParticleTransformTrack::Evaluate(float InCacheTime) const
+FTransform FParticleTransformTrack::Evaluate(float InCacheTime, const FTransform* MassToLocal) const
 {
 	QUICK_SCOPE_CYCLE_COUNTER(QSTAT_EvalParticleTransformTrack);
 	const int32 NumKeys = GetNumKeys();
 
+	FTransform Result{FTransform::Identity};
 	if(NumKeys > 0)
 	{
-		if(InCacheTime < BeginOffset)
-		{
-			// Take first key
-			return FTransform(FQuat(RawTransformTrack.RotKeys[0]), FVector(RawTransformTrack.PosKeys[0]));
-		}
-		else if(InCacheTime > KeyTimestamps.Last())
-		{
-			// Take last key
-			return FTransform(FQuat(RawTransformTrack.RotKeys.Last()), FVector(RawTransformTrack.PosKeys.Last()));
-		}
-		else
-		{
-			// Valid in-range, evaluate
-			if(NumKeys == 1)
-			{
-				return FTransform(FQuat(RawTransformTrack.RotKeys[0]), FVector(RawTransformTrack.PosKeys[0]));
-			}
-
-			// Find the first key with a timestamp greater than InCacheTIme
-			int32 IndexBeyond = 0;
-			{
-				QUICK_SCOPE_CYCLE_COUNTER(QSTAT_UpperBound);
-				IndexBeyond = Algo::UpperBound(KeyTimestamps, InCacheTime);
-			}
-
-			if(IndexBeyond == INDEX_NONE || IndexBeyond >= KeyTimestamps.Num())
-			{
-				// Must be equal to the last key
-				return FTransform(FQuat(RawTransformTrack.RotKeys.Last()), FVector(RawTransformTrack.PosKeys.Last()));
-			}
-
-			const int32 IndexBefore = IndexBeyond - 1;
-
-			if(IndexBefore == INDEX_NONE)
-			{
-				// Must have been equal to first key
-				return FTransform(FQuat(RawTransformTrack.RotKeys[0]), FVector(RawTransformTrack.PosKeys[0]));
-			}
+		const float ClampedTime = FMath::Clamp(InCacheTime, KeyTimestamps[0], KeyTimestamps.Last());  
+		
+		const int32 UpperKeyIndex = GetUpperBoundEvaluationIndex(ClampedTime);
+		const int32 LowerKeyIndex = FMath::Max(0, (UpperKeyIndex-1));;
+		
+		const FTransform LowerTransform = EvaluateAt(LowerKeyIndex);
+		const FTransform UpperTransform = EvaluateAt(UpperKeyIndex);
+		
+		const FTransform TransformA = (MassToLocal)? (*MassToLocal * LowerTransform): LowerTransform;
+		const FTransform TransformB = (MassToLocal)? (*MassToLocal * UpperTransform): UpperTransform;
+		
+		const float Interval = KeyTimestamps[UpperKeyIndex] - KeyTimestamps[LowerKeyIndex];
+		const float Alpha = (Interval > 0)? ((ClampedTime - KeyTimestamps[LowerKeyIndex]) / Interval): 0.0f;
+		ensure(Alpha >= 0.f && Alpha <= 1.f); 
+		
+		Result.Blend(TransformA, TransformB, Alpha); 
+	}
 
-			// Need to interpolate
-			const float Interval = KeyTimestamps[IndexBeyond] - KeyTimestamps[IndexBefore];
-			const float Fraction = (InCacheTime - KeyTimestamps[IndexBefore]) / Interval;
+	return Result;
+}
 
-			// Slerp rotation - lerp translation
-			return FTransform(FQuat(FQuat4f::Slerp(RawTransformTrack.RotKeys[IndexBefore], RawTransformTrack.RotKeys[IndexBeyond], Fraction)),
-							  FVector(FMath::Lerp(RawTransformTrack.PosKeys[IndexBefore], RawTransformTrack.PosKeys[IndexBeyond], Fraction)));
-		}
+int32 FParticleTransformTrack::GetUpperBoundEvaluationIndex(float InCacheTime) const
+{
+	if (KeyTimestamps.Num() == 0)
+	{
+		return 0;
 	}
+	return FMath::Clamp(Algo::UpperBound(KeyTimestamps, InCacheTime), 0, KeyTimestamps.Num()-1);
+}
 
-	return FTransform::Identity;
+FTransform FParticleTransformTrack::EvaluateAt(int32 Index) const
+{
+	if (KeyTimestamps.Num() == 0)
+	{
+		return FTransform::Identity;
+	}
+	return FTransform{FQuat(RawTransformTrack.RotKeys[Index]), FVector(RawTransformTrack.PosKeys[Index])};
 }
 
 const int32 FParticleTransformTrack::GetNumKeys() const
diff --git a/Engine/Plugins/Experimental/ChaosCaching/Source/ChaosCaching/Public/Chaos/ChaosCache.h b/Engine/Plugins/Experimental/ChaosCaching/Source/ChaosCaching/Public/Chaos/ChaosCache.h
--- a/Engine/Plugins/Experimental/ChaosCaching/Source/ChaosCaching/Public/Chaos/ChaosCache.h
+++ b/Engine/Plugins/Experimental/ChaosCaching/Source/ChaosCaching/Public/Chaos/ChaosCache.h
@@ -43,8 +43,25 @@ struct FParticleTransformTrack
 	 * Evaluates the transform track at the specified time, returning the evaluated transform. When in between
 	 * keys translations will be linearly interpolated and rotations spherically interpolated
 	 * @param InCacheTime Absolute time from the beginning of the entire owning cache to evaluate.
+	 * @param MassToLocal if not null, the will be premultiplied to transform before interpolation
 	 */
-	FTransform Evaluate(float InCacheTime) const;
+	FTransform Evaluate(float InCacheTime, const FTransform* MassToLocal) const;
+
+	/**
+	 * Find the index the key where timestamp is directly above InCacheTime
+	 * returned value is garanteed to be within the range of keys if there's any
+	 * in the case there's no key in the track, 0 is returned
+	 * @param InCacheTime Absolute time from the beginning of the entire owning cache to evaluate.
+	 * @return Valid key index
+	 */
+	int32 GetUpperBoundEvaluationIndex(float InCacheTime) const;
+
+	/**
+	 * Evaluate the acche at a sepcific index
+	 * this assumes that index is within the range of valid keys
+	 * * if the track is empty identity transform is returned
+	 */
+	FTransform EvaluateAt(int32 Index) const;
 
 	const int32 GetNumKeys() const;
 	const float GetDuration() const;
@@ -306,9 +323,10 @@ public:
 	/**
 	 * Evaluate the cache with the specified parameters, returning the evaluated results
 	 * @param InContext evaluation context
+	 * @param MassToLocalTransforms MassToLocal trasnform if available ( geometry collection are using them )
 	 * @see FCacheEvaluationContext
 	 */
-	FCacheEvaluationResult Evaluate(const FCacheEvaluationContext& InContext);
+	FCacheEvaluationResult Evaluate(const FCacheEvaluationContext& InContext, const TArray<FTransform>* MassToLocalTransforms);
 
 	/**
 	 * Initializes the spawnable template from a currently existing component so it can be spawned by the editor
@@ -326,12 +344,13 @@ public:
 	 * Evaluates a single particle from the tracks array
 	 * @param InIndex Particle track index (unchecked, ensure valid before call)
 	 * @param InTickRecord Tick record for this evaluation
+	 * @param MassToLocal MassToLocal transform, necessary for proper interpolation, if not available, nullptr shoudl used instead
 	 * @param OutOptTransform Transform to fill, skipped if null
 	 * @param OutOptCurves Curves to fill, skipped if null
 	 */
-	void EvaluateSingle(int32 InIndex, FPlaybackTickRecord& InTickRecord, FTransform* OutOptTransform, TMap<FName, float>* OutOptCurves);
+	void EvaluateSingle(int32 InIndex, FPlaybackTickRecord& InTickRecord, const FTransform* MassToLocal, FTransform* OutOptTransform, TMap<FName, float>* OutOptCurves);
 
-	void EvaluateTransform(const FPerParticleCacheData& InData, float InTime, FTransform& OutTransform);
+	void EvaluateTransform(const FPerParticleCacheData& InData, float InTime, const FTransform* MassToLocal, FTransform& OutTransform);
 	void EvaluateCurves(const FPerParticleCacheData& InData, float InTime, TMap<FName, float>& OutCurves);
 	void EvaluateEvents(FPlaybackTickRecord& InTickRecord, TMap<FName, TArray<FCacheEventHandle>>& OutEvents);
 
diff --git a/Engine/Source/Runtime/Experimental/Chaos/Private/PhysicsProxy/GeometryCollectionPhysicsProxy.cpp b/Engine/Source/Runtime/Experimental/Chaos/Private/PhysicsProxy/GeometryCollectionPhysicsProxy.cpp
--- a/Engine/Source/Runtime/Experimental/Chaos/Private/PhysicsProxy/GeometryCollectionPhysicsProxy.cpp
+++ b/Engine/Source/Runtime/Experimental/Chaos/Private/PhysicsProxy/GeometryCollectionPhysicsProxy.cpp
@@ -2309,7 +2309,12 @@ bool FGeometryCollectionPhysicsProxy::PullFromPhysicsState(const Chaos::FDirtyGe
 						bIsCollectionDirty = true;
 					}
 
-					GameThreadCollection.Transform[TransformGroupIndex].Blend(Prev.Transforms[TransformGroupIndex], Next.Transforms[TransformGroupIndex], *Alpha);
+					
+					// need to interpolate using MassToLocal
+					const FTransform& MassToLocal = GameThreadCollection.MassToLocal[TransformGroupIndex];
+					FTransform InterpolatedTransform;
+					InterpolatedTransform.Blend(MassToLocal * Prev.Transforms[TransformGroupIndex], MassToLocal * Next.Transforms[TransformGroupIndex], *Alpha);
+					GameThreadCollection.Transform[TransformGroupIndex] = MassToLocal.Inverse() * InterpolatedTransform;
 
 					if(LinearVelocities && AngularVelocities)
 					{
