Chaos : fix cache replay for multi level GC where wrong index was used for MassToLocal

diff --git a/Engine/Plugins/Experimental/ChaosCaching/Source/ChaosCaching/Private/Chaos/ChaosCache.cpp b/Engine/Plugins/Experimental/ChaosCaching/Source/ChaosCaching/Private/Chaos/ChaosCache.cpp
--- a/Engine/Plugins/Experimental/ChaosCaching/Source/ChaosCaching/Private/Chaos/ChaosCache.cpp
+++ b/Engine/Plugins/Experimental/ChaosCaching/Source/ChaosCaching/Private/Chaos/ChaosCache.cpp
@@ -7,6 +7,12 @@
 #include "GeometryCollection/GeometryCollectionComponent.h"
 #include "GeometryCollection/GeometryCollectionObject.h"
 
+bool bChaosCacheUseInterpolation = true;
+FAutoConsoleVariableRef CVarChaosCacheUseInterpolation(
+	TEXT("p.Chaos.Cache.UseInterpolation"),
+	bChaosCacheUseInterpolation,
+	TEXT("When enabled, cache interpolates between keys.[def: true]"));
+
 UChaosCache::UChaosCache()
 	: CurrentRecordCount(0)
 	, CurrentPlaybackCount(0)
@@ -365,16 +371,14 @@ FCacheEvaluationResult UChaosCache::Evaluate(const FCacheEvaluationContext& InCo
 					EvalCurves = &Result.Curves[EvalIndex];
 				}
 
+				int32 ParticleIndex = INDEX_NONE;
 				if(TrackToParticle.IsValidIndex(CacheIndex))
 				{
-					Result.ParticleIndices.Add(TrackToParticle[CacheIndex]);
-				}
-				else
-				{
-					Result.ParticleIndices.Add(INDEX_NONE);
+					ParticleIndex = TrackToParticle[CacheIndex];
 				}
+				Result.ParticleIndices.Add(ParticleIndex);
 
-				const FTransform* MassToLocal = MassToLocalTransforms? &((*MassToLocalTransforms)[CacheIndex]): nullptr;
+				const FTransform* MassToLocal = (MassToLocalTransforms && ParticleIndex != INDEX_NONE)? &((*MassToLocalTransforms)[ParticleIndex]): nullptr;
 				EvaluateSingle(CacheIndex, InContext.TickRecord, MassToLocal, EvalTransform, EvalCurves);
 			}
 		}
@@ -422,16 +426,14 @@ FCacheEvaluationResult UChaosCache::Evaluate(const FCacheEvaluationContext& InCo
 				EvalCurves = &Result.Curves.Last();
 			}
 
+			int32 ParticleIndex = INDEX_NONE;
 			if(TrackToParticle.IsValidIndex(Index))
 			{
-				Result.ParticleIndices.Add(TrackToParticle[Index]);
-			}
-			else
-			{
-				Result.ParticleIndices.Add(INDEX_NONE);
+				ParticleIndex = TrackToParticle[Index];
 			}
+			Result.ParticleIndices.Add(ParticleIndex);
 
-			const FTransform* MassToLocal = MassToLocalTransforms? &((*MassToLocalTransforms)[Index]): nullptr;
+			const FTransform* MassToLocal = (MassToLocalTransforms && ParticleIndex != INDEX_NONE)? &((*MassToLocalTransforms)[ParticleIndex]): nullptr;
 			EvaluateSingle(Index, InContext.TickRecord, MassToLocal, EvalTransform, EvalCurves);
 		}
 	}
@@ -558,8 +560,15 @@ FTransform FParticleTransformTrack::Evaluate(float InCacheTime, const FTransform
 		const float Interval = KeyTimestamps[UpperKeyIndex] - KeyTimestamps[LowerKeyIndex];
 		const float Alpha = (Interval > 0)? ((ClampedTime - KeyTimestamps[LowerKeyIndex]) / Interval): 0.0f;
 		ensure(Alpha >= 0.f && Alpha <= 1.f); 
-		
-		Result.Blend(TransformA, TransformB, Alpha); 
+
+		if (bChaosCacheUseInterpolation)
+		{
+			Result.Blend(TransformA, TransformB, Alpha);
+		}
+		else
+		{
+			Result = TransformA;
+		}
 	}
 
 	return Result;
