Chaos cloth support for the chaos cache

diff --git a/Engine/Plugins/Experimental/ChaosCaching/ChaosCaching.uplugin b/Engine/Plugins/Experimental/ChaosCaching/ChaosCaching.uplugin
--- a/Engine/Plugins/Experimental/ChaosCaching/ChaosCaching.uplugin
+++ b/Engine/Plugins/Experimental/ChaosCaching/ChaosCaching.uplugin
@@ -26,5 +26,11 @@
 			"Type": "Editor",
 			"LoadingPhase": "PreDefault"
 		}
+	],
+	"Plugins": [
+		{
+			"Name": "Takes",
+			"Enabled": true
+		}
 	]
 }
\ No newline at end of file
diff --git a/Engine/Plugins/Experimental/ChaosCaching/Source/ChaosCaching/ChaosCaching.Build.cs b/Engine/Plugins/Experimental/ChaosCaching/Source/ChaosCaching/ChaosCaching.Build.cs
--- a/Engine/Plugins/Experimental/ChaosCaching/Source/ChaosCaching/ChaosCaching.Build.cs
+++ b/Engine/Plugins/Experimental/ChaosCaching/Source/ChaosCaching/ChaosCaching.Build.cs
@@ -18,7 +18,10 @@ namespace UnrealBuildTool.Rules
                     "RenderCore",
                     "RHI",
 					"GeometryCollectionEngine",
-					"ChaosSolverEngine"
+					"ChaosSolverEngine",
+					"MovieScene",
+					"MovieSceneTracks",
+					"LevelSequence",
 				});
 
 			if(Target.bBuildEditor)
diff --git a/Engine/Plugins/Experimental/ChaosCaching/Source/ChaosCaching/Private/Chaos/Adapters/GeometryCollectionComponentCacheAdapter.cpp b/Engine/Plugins/Experimental/ChaosCaching/Source/ChaosCaching/Private/Chaos/Adapters/GeometryCollectionComponentCacheAdapter.cpp
--- a/Engine/Plugins/Experimental/ChaosCaching/Source/ChaosCaching/Private/Chaos/Adapters/GeometryCollectionComponentCacheAdapter.cpp
+++ b/Engine/Plugins/Experimental/ChaosCaching/Source/ChaosCaching/Private/Chaos/Adapters/GeometryCollectionComponentCacheAdapter.cpp
@@ -612,6 +612,11 @@ namespace Chaos
 
 		return nullptr;
 	}
+	
+	Chaos::FPhysicsSolverEvents* FGeometryCollectionCacheAdapter::BuildEventsSolver(UPrimitiveComponent* InComponent) const
+	{
+		return GetComponentSolver(InComponent);
+	}
 
 	void FGeometryCollectionCacheAdapter::SetRestState(UPrimitiveComponent* InComponent, UChaosCache* InCache, const FTransform& InRootTransform, Chaos::FReal InTime) const
 	{
@@ -630,7 +635,6 @@ namespace Chaos
 			const FTransform ActorToWorld = InRootTransform;
 			const FTransform WorldToComponent = Comp->GetComponentTransform().Inverse();
 			FTransform ActorToComponent = ActorToWorld * WorldToComponent;
-
 			
 			if (const UGeometryCollection* RestCollection = Comp->GetRestCollection())
 			{
diff --git a/Engine/Plugins/Experimental/ChaosCaching/Source/ChaosCaching/Private/Chaos/Adapters/StaticMeshComponentCacheAdapter.cpp b/Engine/Plugins/Experimental/ChaosCaching/Source/ChaosCaching/Private/Chaos/Adapters/StaticMeshComponentCacheAdapter.cpp
--- a/Engine/Plugins/Experimental/ChaosCaching/Source/ChaosCaching/Private/Chaos/Adapters/StaticMeshComponentCacheAdapter.cpp
+++ b/Engine/Plugins/Experimental/ChaosCaching/Source/ChaosCaching/Private/Chaos/Adapters/StaticMeshComponentCacheAdapter.cpp
@@ -136,7 +136,12 @@ namespace Chaos
 
 		return nullptr;
 	}
-
+	
+	Chaos::FPhysicsSolverEvents* FStaticMeshCacheAdapter::BuildEventsSolver(UPrimitiveComponent* InComponent) const
+	{
+		return GetComponentSolver(InComponent);
+	}
+	
 	void FStaticMeshCacheAdapter::SetRestState(UPrimitiveComponent* InComponent, UChaosCache* InCache, const FTransform& InRootTransform, Chaos::FReal InTime) const
 	{
 		if (!InCache || InCache->GetDuration() == 0.0f)
diff --git a/Engine/Plugins/Experimental/ChaosCaching/Source/ChaosCaching/Private/Chaos/CacheCollection.cpp b/Engine/Plugins/Experimental/ChaosCaching/Source/ChaosCaching/Private/Chaos/CacheCollection.cpp
--- a/Engine/Plugins/Experimental/ChaosCaching/Source/ChaosCaching/Private/Chaos/CacheCollection.cpp
+++ b/Engine/Plugins/Experimental/ChaosCaching/Source/ChaosCaching/Private/Chaos/CacheCollection.cpp
@@ -87,7 +87,12 @@ void UChaosCacheCollection::FlushAllCacheWrites()
 	});
 }
 
-const TArray<UChaosCache*> UChaosCacheCollection::GetCaches() const
+float UChaosCacheCollection::GetMaxDuration() const
 {
-	return Caches;
-}
+	float MaxDuration = 0.0;
+    for( UChaosCache* CacheInstance : Caches)
+    {
+    	MaxDuration = FMath::Max(CacheInstance->GetDuration(), MaxDuration);
+    }
+	return MaxDuration;
+}
\ No newline at end of file
diff --git a/Engine/Plugins/Experimental/ChaosCaching/Source/ChaosCaching/Private/Chaos/CacheManagerActor.cpp b/Engine/Plugins/Experimental/ChaosCaching/Source/ChaosCaching/Private/Chaos/CacheManagerActor.cpp
--- a/Engine/Plugins/Experimental/ChaosCaching/Source/ChaosCaching/Private/Chaos/CacheManagerActor.cpp
+++ b/Engine/Plugins/Experimental/ChaosCaching/Source/ChaosCaching/Private/Chaos/CacheManagerActor.cpp
@@ -41,14 +41,26 @@ void FObservedComponent::ResetRuntimeData(const EStartMode ManagerStartMode)
 	TickRecord.Reset();
 }
 
+void FObservedComponent::PostSerialize(const FArchive& Ar)
+{
+	PRAGMA_DISABLE_DEPRECATION_WARNINGS
+	if (Ar.IsLoading() && ComponentRef.OtherActor)
+	{
+		SoftComponentRef.OtherActor = ComponentRef.OtherActor;
+		SoftComponentRef.ComponentProperty = ComponentRef.ComponentProperty;
+		SoftComponentRef.OverrideComponent = ComponentRef.OverrideComponent;
+		SoftComponentRef.PathToComponent = ComponentRef.PathToComponent; 
+	}
+	PRAGMA_ENABLE_DEPRECATION_WARNINGS
+}
 UPrimitiveComponent* FObservedComponent::GetComponent()
 {
-	return Cast<UPrimitiveComponent>(ComponentRef.GetComponent(nullptr));
+	return Cast<UPrimitiveComponent>(SoftComponentRef.GetComponent(nullptr));
 }
 
 UPrimitiveComponent* FObservedComponent::GetComponent() const
 {
-	return Cast<UPrimitiveComponent>(ComponentRef.GetComponent(nullptr));
+	return Cast<UPrimitiveComponent>(SoftComponentRef.GetComponent(nullptr));
 }
 
 AChaosCacheManager::AChaosCacheManager(const FObjectInitializer& ObjectInitializer /*= FObjectInitializer::Get()*/)
@@ -145,6 +157,16 @@ void AChaosCacheManager::PostEditChangeProperty(FPropertyChangedEvent& PropertyC
 }
 #endif
 
+void AChaosCacheManager::Serialize(FArchive& Ar)
+{
+	Super::Serialize(Ar);
+
+	for(FObservedComponent& ObervedComponent : ObservedComponents)
+	{
+		ObervedComponent.PostSerialize(Ar);
+	}
+}
+
 void AChaosCacheManager::SetStartTime(float InStartTime)
 {
 	StartTime = InStartTime;
@@ -207,12 +229,19 @@ void AChaosCacheManager::SelectComponent(int32 InIndex)
 }
 #endif
 
+
 void AChaosCacheManager::BeginPlay()
 {
 	using namespace Chaos;
 
 	Super::BeginPlay();
 
+	BeginEvaluate();
+}
+
+void AChaosCacheManager::BeginEvaluate()
+{
+	using namespace Chaos;
 	bIsSimulating = !(CacheMode == ECacheMode::None);
 	
 	if(!CacheCollection)
@@ -228,8 +257,6 @@ void AChaosCacheManager::BeginPlay()
 		SetActorTickEnabled(true);
 	}
 
-	const TArray<UChaosCache*> Caches = CacheCollection->GetCaches();
-
 	// Build list of adapters for our observed components
 	IModularFeatures&               ModularFeatures = IModularFeatures::Get();
 	TArray<FComponentCacheAdapter*> Adapters        = ModularFeatures.GetModularFeatureImplementations<FComponentCacheAdapter>(FComponentCacheAdapter::FeatureName);
@@ -238,7 +265,6 @@ void AChaosCacheManager::BeginPlay()
 	{
 		Adapter->Initialize();
 	}
-	
 	ActiveAdapters.Reset();
 
 	int32       NumFailedPlaybackEntries = 0;
@@ -250,18 +276,17 @@ void AChaosCacheManager::BeginPlay()
 		auto ByPriority = [](const FComponentCacheAdapter* A, const FComponentCacheAdapter* B) {
 			return A->GetPriority() < B->GetPriority();
 		};
-
 		UPrimitiveComponent* Comp = Observed.GetComponent();
-
+		
 		if(!Comp)
 		{
 			UE_LOG(LogChaosCache, Warning, TEXT("%s has invalid observed component reference."), *GetName());
 			ActiveAdapters.Add(nullptr);
 			continue;
 		}
-
+		
 		UClass* ActualClass = Comp->GetClass();
-
+    
 		TArray<FComponentCacheAdapter*> DirectAdapters = Adapters.FilterByPredicate([ActualClass](const FComponentCacheAdapter* InTest) {
 			return InTest && InTest->SupportsComponentClass(ActualClass) == Chaos::FComponentCacheAdapter::SupportType::Direct;
 		});
@@ -294,97 +319,96 @@ void AChaosCacheManager::BeginPlay()
 
 		// Reset timers and last cache
 		Observed.ResetRuntimeData(StartMode);
-
+		
 		bool                    bRequiresRecord = false;
 		FComponentCacheAdapter* CurrAdapter     = ActiveAdapters[Index];
 		check(CurrAdapter);    // should definitely have added one above
 
-		if(Chaos::FPhysicsSolver* Solver = CurrAdapter->GetComponentSolver(Comp))
+		// We only need to register callbacks if in play/record mode
+		if(CacheMode != ECacheMode::None)
 		{
-			FPerSolverData* SolverData = PerSolverData.Find(Solver);
-
-			if(!SolverData)
+			if(Chaos::FPhysicsSolverEvents* Solver = CurrAdapter->BuildEventsSolver(Comp))
 			{
-				SolverData                  = &PerSolverData.Add(Solver);
-				SolverData->PreSolveHandle  = Solver->AddPreAdvanceCallback(FSolverPreAdvance::FDelegate::CreateUObject(this, &AChaosCacheManager::HandlePreSolve, Solver));
-				SolverData->PreBufferHandle = Solver->AddPreBufferCallback(FSolverPreAdvance::FDelegate::CreateUObject(this, &AChaosCacheManager::HandlePreBuffer, Solver));
-				SolverData->PostSolveHandle = Solver->AddPostAdvanceCallback(FSolverPostAdvance::FDelegate::CreateUObject(this, &AChaosCacheManager::HandlePostSolve, Solver));
-				SolverData->TeardownHandle  = Solver->AddTeardownCallback(FSolverTeardown::FDelegate::CreateUObject(this, &AChaosCacheManager::HandleTeardown, Solver));
-			}
+				FPerSolverData* SolverData = PerSolverData.Find(Solver);
 
-			if (CacheMode == ECacheMode::Play || CacheMode == ECacheMode::None || !CanRecord())
-			{
-				UChaosCache*    PlayCache = CacheCollection->FindCache(Observed.CacheName);
-					
-				if(PlayCache)
+				if(!SolverData)
 				{
-					FCacheUserToken Token = PlayCache->BeginPlayback();
+					SolverData                  = &PerSolverData.Add(Solver);
+					SolverData->PreSolveHandle  = Solver->AddPreAdvanceCallback(FSolverPreAdvance::FDelegate::CreateUObject(this, &AChaosCacheManager::HandlePreSolve, Solver));
+					SolverData->PreBufferHandle = Solver->AddPreBufferCallback(FSolverPreAdvance::FDelegate::CreateUObject(this, &AChaosCacheManager::HandlePreBuffer, Solver));
+					SolverData->PostSolveHandle = Solver->AddPostAdvanceCallback(FSolverPostAdvance::FDelegate::CreateUObject(this, &AChaosCacheManager::HandlePostSolve, Solver));
+					SolverData->TeardownHandle  = Solver->AddTeardownCallback(FSolverTeardown::FDelegate::CreateUObject(this, &AChaosCacheManager::HandleTeardown, Solver));
+				}
 
-					if(Token.IsOpen() && CurrAdapter->ValidForPlayback(Comp, PlayCache))
+				if (CacheMode == ECacheMode::Play || !CanRecord())
+				{ 
+					if(UChaosCache*    PlayCache = CacheCollection->FindCache(Observed.CacheName))
 					{
-						if (CacheMode == ECacheMode::Play)
+						FCacheUserToken Token = PlayCache->BeginPlayback();
+						if(Token.IsOpen() && CurrAdapter->ValidForPlayback(Comp, PlayCache))
+						{
+							if (CacheMode == ECacheMode::Play)
+							{
+								SolverData->PlaybackIndices.Add(Index);
+								SolverData->PlaybackTickRecords.AddDefaulted();
+								SolverData->PlaybackTickRecords.Last().SetSpaceTransform(GetTransform());
+								Observed.Cache = PlayCache;
+								Observed.TickRecord.SetSpaceTransform(GetTransform());
+								Observed.TickRecord.SetLastTime(StartTime);
+								OpenPlaybackCaches.Add(TTuple<FCacheUserToken, UChaosCache*>(MoveTemp(Token), Observed.Cache));
+								CurrAdapter->InitializeForPlayback(Comp, Observed.Cache, StartTime);
+							}
+						}
+						else
 						{
-							SolverData->PlaybackIndices.Add(Index);
-							SolverData->PlaybackTickRecords.AddDefaulted();
-							SolverData->PlaybackTickRecords.Last().SetSpaceTransform(GetTransform());
-							Observed.Cache = PlayCache;
-							Observed.TickRecord.SetSpaceTransform(GetTransform());
-							Observed.TickRecord.SetLastTime(StartTime);
-							OpenPlaybackCaches.Add(TTuple<FCacheUserToken, UChaosCache*>(MoveTemp(Token), Observed.Cache));
-							CurrAdapter->InitializeForPlayback(Comp, Observed.Cache, StartTime);
+							if(Token.IsOpen())
+							{
+								UE_LOG(LogChaosCache,
+										Warning,
+										TEXT("Failed playback for component %s, Selected cache adapter unable to handle the cache (the cache is incompatible)"),
+										*Comp->GetPathName());
+
+								// The cache session was valid so make sure to end it
+								PlayCache->EndPlayback(Token);
+							}
+							else    // Already open for record somewhere
+								{
+								UE_LOG(LogChaosCache,
+										Warning,
+										TEXT("Failed playback for component %s using cache %s, cache already open for record"),
+										*Comp->GetName(),
+										*PlayCache->GetPathName());
+								}
+
+							++NumFailedPlaybackEntries;
 						}
 					}
 					else
 					{
-						if(Token.IsOpen())
-						{
-							UE_LOG(LogChaosCache,
-									Warning,
-									TEXT("Failed playback for component %s, Selected cache adapter unable to handle the cache (the cache is incompatible)"),
-									*Comp->GetPathName());
-
-							// The cache session was valid so make sure to end it
-							PlayCache->EndPlayback(Token);
-						}
-						else    // Already open for record somewhere
-						{
-							UE_LOG(LogChaosCache,
-									Warning,
-									TEXT("Failed playback for component %s using cache %s, cache already open for record"),
-									*Comp->GetName(),
-									*PlayCache->GetPathName());
-						}
-
-						++NumFailedPlaybackEntries;
+						UE_LOG(LogChaosCache, Log, TEXT("Skipping playback for component %s, no available cache."), *Comp->GetName());
 					}
 				}
-				else
+				else // CacheMode == ECacheMode::Record
 				{
-					UE_LOG(LogChaosCache, Log, TEXT("Skipping playback for component %s, no available cache."), *Comp->GetName());
-				}
-			}
-			else // CacheMode == ECacheMode::Record
-			{
-				// Make sure there's a cache available if we're going to record.
-				UPrimitiveComponent* Component = Observed.GetComponent();
-				FName CacheName = (Observed.CacheName == NAME_None) ? MakeUniqueObjectName(CacheCollection, UChaosCache::StaticClass(), "Cache") : Observed.CacheName;
+					FName CacheName = (Observed.CacheName == NAME_None) ? MakeUniqueObjectName(CacheCollection, UChaosCache::StaticClass(), "Cache") : Observed.CacheName;
 
-				UChaosCache*    RecordCache = CacheCollection->FindOrAddCache(CacheName);
-				FCacheUserToken Token       = RecordCache->BeginRecord(Observed.GetComponent(), CurrAdapter->GetGuid(), GetTransform());
+					UChaosCache*    RecordCache = CacheCollection->FindOrAddCache(CacheName);
+					FCacheUserToken Token       = RecordCache->BeginRecord(Comp, CurrAdapter->GetGuid(), GetTransform());
 
-				if(Token.IsOpen())
-				{
-					SolverData->RecordIndices.Add(Index);
+					if(Token.IsOpen())
+					{
+						SolverData->RecordIndices.Add(Index);
 
-					Observed.Cache = CacheCollection->FindOrAddCache(CacheName);
-						
-					// We'll record the observed component in Cache Manager's local space.
-					Observed.TickRecord.SetSpaceTransform(GetTransform());
-					OpenRecordCaches.Add(TTuple<FCacheUserToken, UChaosCache*>(MoveTemp(Token), Observed.Cache));
-					CurrAdapter->InitializeForRecord(Component, Observed.Cache);
+						Observed.Cache = CacheCollection->FindOrAddCache(CacheName);
+				
+						// We'll record the observed component in Cache Manager's local space.
+						Observed.TickRecord.SetSpaceTransform(GetTransform());
+						OpenRecordCaches.Add(TTuple<FCacheUserToken, UChaosCache*>(MoveTemp(Token), Observed.Cache));
+						CurrAdapter->InitializeForRecord(Comp, Observed.Cache);
 
-					// Ensure we enable the actor tick to flush out the pending record writes
-					bRequiresRecord = true;
+						// Ensure we enable the actor tick to flush out the pending record writes
+						bRequiresRecord =  true;
+					}
 				}
 			}
 		}
@@ -412,26 +436,34 @@ void AChaosCacheManager::EndPlay(const EEndPlayReason::Type EndPlayReason)
 {
 	Super::EndPlay(EndPlayReason);
 
+	EndEvaluate();
+}
+
+void AChaosCacheManager::EndEvaluate()
+{
 	bIsSimulating = false;
 	
 	FChaosSolversModule* Module = FChaosSolversModule::GetModule();
 	check(Module);
 
-	for(TPair<Chaos::FPhysicsSolver*, FPerSolverData> PerSolver : PerSolverData)
+	for(TPair<Chaos::FPhysicsSolverEvents*, FPerSolverData> PerSolver : PerSolverData)
 	{
-		Chaos::FPhysicsSolver* CurrSolver = PerSolver.Key;
+		Chaos::FPhysicsSolverEvents* CurrSolver = PerSolver.Key;
 		FPerSolverData&        CurrData   = PerSolver.Value;
 		if(ensure(CurrSolver))
 		{
 			ensure(CurrSolver->RemovePostAdvanceCallback(CurrData.PostSolveHandle));
 			ensure(CurrSolver->RemovePreBufferCallback(CurrData.PreBufferHandle));
 			ensure(CurrSolver->RemovePreAdvanceCallback(CurrData.PreSolveHandle));
+			ensure(CurrSolver->RemoveTeardownCallback(CurrData.TeardownHandle));
+
 			CurrData.PostSolveHandle.Reset();
 			CurrData.PreSolveHandle.Reset();
 			CurrData.PreBufferHandle.Reset();
+			CurrData.TeardownHandle.Reset();
 		}
 	}
-
+	PerSolverData.Reset();
 	ActiveAdapters.Reset();
 
 	// Close any open caches as the session is complete. this will flush pending writes and post-process the cache
@@ -448,7 +480,7 @@ void AChaosCacheManager::EndPlay(const EEndPlayReason::Type EndPlayReason)
 	OpenPlaybackCaches.Reset();
 }
 
-void AChaosCacheManager::HandlePreSolve(Chaos::FReal InDt, Chaos::FPhysicsSolver* InSolver)
+void AChaosCacheManager::HandlePreSolve(Chaos::FReal InDt, Chaos::FPhysicsSolverEvents* InSolver)
 {
 	if(!CacheCollection)
 	{
@@ -459,7 +491,7 @@ void AChaosCacheManager::HandlePreSolve(Chaos::FReal InDt, Chaos::FPhysicsSolver
 
 	if(!Data)
 	{
-		ensureMsgf(false, TEXT("AChaosCacheManager::HandlePostSolve couldn't find a solver entry - a solver binding has leaked."));
+		ensureMsgf(false, TEXT("AChaosCacheManager::HandlePreSolve couldn't find a solver entry - a solver binding has leaked."));
 		return;
 	}
 
@@ -472,7 +504,7 @@ void AChaosCacheManager::HandlePreSolve(Chaos::FReal InDt, Chaos::FPhysicsSolver
 	});
 }
 
-void AChaosCacheManager::HandlePreBuffer(Chaos::FReal InDt, Chaos::FPhysicsSolver* InSolver)
+void AChaosCacheManager::HandlePreBuffer(Chaos::FReal InDt, Chaos::FPhysicsSolverEvents* InSolver)
 {
 	if(!CacheCollection)
 	{
@@ -486,16 +518,19 @@ void AChaosCacheManager::HandlePreBuffer(Chaos::FReal InDt, Chaos::FPhysicsSolve
 		ensureMsgf(false, TEXT("AChaosCacheManager::HandlePreBuffer couldn't find a solver entry - a solver binding has leaked."));
 		return;
 	}
-
-	for(Chaos::TPBDRigidParticleHandle<Chaos::FReal, 3>* PendingKinematic : Data->PendingKinematicUpdates)
+	if(Data->PendingKinematicUpdates.Num() > 0)
 	{
-		InSolver->GetParticles().MarkTransientDirtyParticle(PendingKinematic);
+		// If the pending rigid kinematic particles are not empty we are dealing with a rigid solver
+		Chaos::FPhysicsSolver* PhysicsSolver = StaticCast<Chaos::FPhysicsSolver*>(InSolver);
+		for(Chaos::TPBDRigidParticleHandle<Chaos::FReal, 3>* PendingKinematic : Data->PendingKinematicUpdates)
+		{
+			PhysicsSolver->GetParticles().MarkTransientDirtyParticle(PendingKinematic);
+		}
+		Data->PendingKinematicUpdates.Reset();
 	}
-
-	Data->PendingKinematicUpdates.Reset();
 }
 
-void AChaosCacheManager::HandlePostSolve(Chaos::FReal InDt, Chaos::FPhysicsSolver* InSolver)
+void AChaosCacheManager::HandlePostSolve(Chaos::FReal InDt, Chaos::FPhysicsSolverEvents* InSolver)
 {
 	if(!CacheCollection)
 	{
@@ -514,6 +549,7 @@ void AChaosCacheManager::HandlePostSolve(Chaos::FReal InDt, Chaos::FPhysicsSolve
 		UPrimitiveComponent* Comp = Observed.GetComponent();
 		if (ensure(Comp && InCache))
 		{
+			
 			// If we haven't advanced since the last record, don't push another frame
 			if (Observed.TimeSinceTrigger > InCache->GetDuration())
 			{
@@ -527,7 +563,7 @@ void AChaosCacheManager::HandlePostSolve(Chaos::FReal InDt, Chaos::FPhysicsSolve
 	});
 }
 
-void AChaosCacheManager::HandleTeardown(Chaos::FPhysicsSolver* InSolver)
+void AChaosCacheManager::HandleTeardown(Chaos::FPhysicsSolverEvents* InSolver)
 {
 	// Solver has been deleted, we remove all reference to it to avoid dereferencing a dangling pointer.
 	PerSolverData.Remove(InSolver);
@@ -544,7 +580,6 @@ void AChaosCacheManager::OnStartFrameChanged(Chaos::FReal InTime)
 	{
 		return;
 	}
-	
 	for (int32 ObservedIdx = 0; ObservedIdx < ObservedComponents.Num(); ++ObservedIdx)
 	{
 		FObservedComponent& Observed = ObservedComponents[ObservedIdx];
@@ -628,8 +663,8 @@ FObservedComponent& AChaosCacheManager::AddNewObservedComponent(UPrimitiveCompon
 	ObservedComponents.AddDefaulted();
 	FObservedComponent& NewEntry = ObservedComponents.Last();
 
-	NewEntry.ComponentRef.PathToComponent = InComponent->GetPathName(InComponent->GetOwner());
-	NewEntry.ComponentRef.OtherActor      = InComponent->GetOwner();
+	NewEntry.SoftComponentRef.PathToComponent = InComponent->GetPathName(InComponent->GetOwner());
+	NewEntry.SoftComponentRef.OtherActor      = InComponent->GetOwner();
 
 	FName CacheName = NAME_None;
 #if WITH_EDITOR
diff --git a/Engine/Plugins/Experimental/ChaosCaching/Source/ChaosCaching/Private/Chaos/ChaosCache.cpp b/Engine/Plugins/Experimental/ChaosCaching/Source/ChaosCaching/Private/Chaos/ChaosCache.cpp
--- a/Engine/Plugins/Experimental/ChaosCaching/Source/ChaosCaching/Private/Chaos/ChaosCache.cpp
+++ b/Engine/Plugins/Experimental/ChaosCaching/Source/ChaosCaching/Private/Chaos/ChaosCache.cpp
@@ -111,6 +111,30 @@ void UChaosCache::FlushPendingFrames()
 	while(PendingWrites.Dequeue(NewData))
 	{
 		QUICK_SCOPE_CYCLE_COUNTER(QSTAT_CacheFlushSingleFrame);
+
+		{
+			const int32 NumParticles = NewData.PendingChannelsIndices.Num();
+			bWroteParticleData |= NumParticles > 0;
+
+			for(auto& ChannelData : NewData.PendingChannelsData)
+			{
+				// TArray<FRichCurve>& TargetCurve = ChannelsTracks.FindOrAdd(ChannelData.Key).RichCurves;
+				// TargetCurve.SetNum(NumParticles);
+
+				FRichCurves* TargetCurve = ChannelsTracks.Find(ChannelData.Key);
+				if(TargetCurve == nullptr)
+				{
+					TargetCurve = &ChannelsTracks.Add(ChannelData.Key);
+					TargetCurve->RichCurves.SetNum(NumParticles);
+				}
+				int32 ParticleIndex = 0;
+				for(auto& ChannelValue : ChannelData.Value)
+				{
+					TargetCurve->RichCurves[ParticleIndex++].AddKey(NewData.Time, ChannelValue);
+				}
+			}
+		}
+		
 		const int32 ParticleCount = NewData.PendingParticleData.Num();
 
 		bWroteParticleData |= ParticleCount > 0;
@@ -136,7 +160,7 @@ void UChaosCache::FlushPendingFrames()
 				PTrack.BeginOffset = NewData.Time;
 
 				// Particle will hold at end of recording.
-				PTrack.bDeactivateOnEnd = false;
+				PTrack.bDeactivateOnEnd = false; 
 			}
 
 			if (ParticleData.bPendingDeactivate)
@@ -208,6 +232,7 @@ FCacheUserToken UChaosCache::BeginRecord(UPrimitiveComponent* InComponent, FGuid
 			RecordedDuration = 0.0f;
 			NumRecordedFrames = 0;
 			ParticleTracks.Reset();
+			ChannelsTracks.Reset();
 			TrackToParticle.Reset();
 			CurveData.Reset();
 			EventTracks.Reset();
@@ -323,6 +348,7 @@ FCacheEvaluationResult UChaosCache::Evaluate(const FCacheEvaluationContext& InCo
 	QUICK_SCOPE_CYCLE_COUNTER(QSTAT_CacheEval);
 
 	FCacheEvaluationResult Result;
+	static const TMap<FName, float> EmptyMap;
 
 	if(CurrentPlaybackCount.load() == 0)
 	{
@@ -333,7 +359,7 @@ FCacheEvaluationResult UChaosCache::Evaluate(const FCacheEvaluationContext& InCo
 
 	Result.EvaluatedTime = InContext.TickRecord.GetTime();
 
-	if(!InContext.bEvaluateTransform && !InContext.bEvaluateCurves && !InContext.bEvaluateEvents)
+	if(!InContext.bEvaluateTransform && !InContext.bEvaluateCurves && !InContext.bEvaluateEvents && !InContext.bEvaluateChannels)
 	{
 		// no evaluation requested
 		return Result;
@@ -345,17 +371,32 @@ FCacheEvaluationResult UChaosCache::Evaluate(const FCacheEvaluationContext& InCo
 	{
 		if(InContext.bEvaluateTransform)
 		{
-			Result.Transform.SetNum(NumProvidedIndices);
+			Result.Transform.Init(FTransform::Identity, NumProvidedIndices);
 		}
 
 		if(InContext.bEvaluateCurves)
 		{
-			Result.Curves.SetNum(NumProvidedIndices);
+			Result.Curves.Init(EmptyMap,NumProvidedIndices);
+		}
+
+		if(InContext.bEvaluateChannels)
+		{
+			for(auto& ChannelData : ChannelsTracks)
+			{
+				TArray<float>& ResultData = Result.Channels.FindOrAdd(ChannelData.Key);
+				ResultData.Init(0.0, NumProvidedIndices);
+
+				for(int32 EvalIndex = 0; EvalIndex < NumProvidedIndices; ++EvalIndex)
+				{
+					const int32 CacheIndex = InContext.EvaluationIndices[EvalIndex];
+					ResultData[CacheIndex] = ChannelData.Value.RichCurves[CacheIndex].Eval(InContext.TickRecord.GetTime(),0.0);
+				}
+			}
 		}
 
 		for(int32 EvalIndex = 0; EvalIndex < NumProvidedIndices; ++EvalIndex)
 		{
-			int32 CacheIndex = InContext.EvaluationIndices[EvalIndex];
+			const int32 CacheIndex = InContext.EvaluationIndices[EvalIndex];
 			if(ensure(ParticleTracks.IsValidIndex(CacheIndex)))
 			{
 				FTransform* EvalTransform = nullptr;
@@ -396,7 +437,19 @@ FCacheEvaluationResult UChaosCache::Evaluate(const FCacheEvaluationContext& InCo
 		{
 			Result.Curves.Reserve(NumParticles);
 		}
-
+		if(InContext.bEvaluateChannels)
+		{
+			for(auto& ChannelData : ChannelsTracks)
+			{
+				TArray<float>& ResultData = Result.Channels.FindOrAdd(ChannelData.Key);
+				ResultData.Init(0.0, NumParticles);
+				
+				for(int32 EvalIndex = 0; EvalIndex < NumParticles; ++EvalIndex)
+				{
+					ResultData[EvalIndex] = ChannelData.Value.RichCurves[EvalIndex].Eval(InContext.TickRecord.GetTime(),0.0);
+				}
+			}
+		}
 		for(int32 Index = 0; Index < NumParticles; ++Index)
 		{
 			if(ParticleTracks[Index].TransformData.BeginOffset > InContext.TickRecord.GetTime())
@@ -416,13 +469,13 @@ FCacheEvaluationResult UChaosCache::Evaluate(const FCacheEvaluationContext& InCo
 
 			if(InContext.bEvaluateTransform)
 			{
-				Result.Transform.AddUninitialized();
+				Result.Transform.Add(FTransform::Identity);
 				EvalTransform = &Result.Transform.Last();
 			}
 
 			if(InContext.bEvaluateCurves)
 			{
-				Result.Curves.AddUninitialized();
+				Result.Curves.Add(EmptyMap);
 				EvalCurves = &Result.Curves.Last();
 			}
 
diff --git a/Engine/Plugins/Experimental/ChaosCaching/Source/ChaosCaching/Private/Chaos/ChaosCachingPlugin.cpp b/Engine/Plugins/Experimental/ChaosCaching/Source/ChaosCaching/Private/Chaos/ChaosCachingPlugin.cpp
--- a/Engine/Plugins/Experimental/ChaosCaching/Source/ChaosCaching/Private/Chaos/ChaosCachingPlugin.cpp
+++ b/Engine/Plugins/Experimental/ChaosCaching/Source/ChaosCaching/Private/Chaos/ChaosCachingPlugin.cpp
@@ -3,6 +3,8 @@
 #include "Chaos/ChaosCachingPlugin.h"
 #include "Chaos/Adapters/GeometryCollectionComponentCacheAdapter.h"
 #include "Chaos/Adapters/StaticMeshComponentCacheAdapter.h"
+#include "Chaos/Sequencer/ChaosCacheObjectSpawner.h"
+#include "ILevelSequenceModule.h"
 #include "CoreMinimal.h"
 #include "Modules/ModuleManager.h"
 
@@ -15,6 +17,9 @@ void IChaosCachingPlugin::StartupModule()
 	GeometryCollectionAdapter = MakeUnique<Chaos::FGeometryCollectionCacheAdapter>();
 	StaticMeshAdapter = MakeUnique<Chaos::FStaticMeshCacheAdapter>();
 
+	ILevelSequenceModule& LevelSequenceModule = FModuleManager::LoadModuleChecked<ILevelSequenceModule>("LevelSequence");
+	OnCreateMovieSceneObjectSpawnerHandle = LevelSequenceModule.RegisterObjectSpawner(FOnCreateMovieSceneObjectSpawner::CreateStatic(&FChaosCacheObjectSpawner::CreateObjectSpawner));
+	
 	RegisterAdapter(GeometryCollectionAdapter.Get());
 	RegisterAdapter(StaticMeshAdapter.Get());
 }
@@ -26,4 +31,10 @@ void IChaosCachingPlugin::ShutdownModule()
 
 	StaticMeshAdapter = nullptr;
 	GeometryCollectionAdapter = nullptr;
+
+	ILevelSequenceModule* LevelSequenceModule = FModuleManager::GetModulePtr<ILevelSequenceModule>("LevelSequence");
+	if (LevelSequenceModule)
+	{
+		LevelSequenceModule->UnregisterObjectSpawner(OnCreateMovieSceneObjectSpawnerHandle);
+	}
 }
diff --git a/Engine/Plugins/Experimental/ChaosCaching/Source/ChaosCaching/Private/Chaos/Sequencer/ChaosCacheObjectSpawner.cpp b/Engine/Plugins/Experimental/ChaosCaching/Source/ChaosCaching/Private/Chaos/Sequencer/ChaosCacheObjectSpawner.cpp
new file mode 100644
--- /dev/null
+++ b/Engine/Plugins/Experimental/ChaosCaching/Source/ChaosCaching/Private/Chaos/Sequencer/ChaosCacheObjectSpawner.cpp
@@ -0,0 +1,53 @@
+﻿// Copyright Epic Games, Inc. All Rights Reserved.
+
+#include "Chaos/Sequencer/ChaosCacheObjectSpawner.h"
+#include "Chaos/CacheManagerActor.h"
+#include "MovieSceneSpawnable.h"
+#include "UObject/Package.h"
+#include "Engine/World.h"
+
+FChaosCacheObjectSpawner::FChaosCacheObjectSpawner() : FLevelSequenceActorSpawner()
+{}
+
+TSharedRef<IMovieSceneObjectSpawner> FChaosCacheObjectSpawner::CreateObjectSpawner()
+{
+	return MakeShareable(new FChaosCacheObjectSpawner);
+}
+
+UClass* FChaosCacheObjectSpawner::GetSupportedTemplateType() const
+{
+	return AChaosCacheManager::StaticClass();
+}
+
+UObject* FChaosCacheObjectSpawner::SpawnObject(FMovieSceneSpawnable& Spawnable, FMovieSceneSequenceIDRef TemplateID, IMovieScenePlayer& Player)
+{
+	UObject* SpawnedObject = FLevelSequenceActorSpawner::SpawnObject(Spawnable, TemplateID, Player);
+	
+	if (AChaosCacheManager* ChaosCache = Cast<AChaosCacheManager>(SpawnedObject))
+	{
+		for(FObservedComponent& ObservedComponent : ChaosCache->GetObservedComponents())
+		{
+			FString FullPath = ObservedComponent.SoftComponentRef.OtherActor.ToString();
+
+#if WITH_EDITORONLY_DATA
+			if (const UPackage* ObjectPackage = SpawnedObject->GetPackage())
+			{
+				// If this is being set from PIE we need to remove the pie prefix 
+				if (ObjectPackage->GetPIEInstanceID() == INDEX_NONE)
+				{
+					int32 PIEInstanceID = INDEX_NONE;
+					FullPath = UWorld::RemovePIEPrefix(FullPath, &PIEInstanceID);
+				}
+			}
+#endif
+			ObservedComponent.SoftComponentRef.OtherActor = FSoftObjectPath(FullPath);
+		}
+	}
+
+	return SpawnedObject;
+}
+
+void FChaosCacheObjectSpawner::DestroySpawnedObject(UObject& Object)
+{
+	FLevelSequenceActorSpawner::DestroySpawnedObject(Object);
+}
diff --git a/Engine/Plugins/Experimental/ChaosCaching/Source/ChaosCaching/Public/Chaos/Adapters/CacheAdapter.h b/Engine/Plugins/Experimental/ChaosCaching/Source/ChaosCaching/Public/Chaos/Adapters/CacheAdapter.h
--- a/Engine/Plugins/Experimental/ChaosCaching/Source/ChaosCaching/Public/Chaos/Adapters/CacheAdapter.h
+++ b/Engine/Plugins/Experimental/ChaosCaching/Source/ChaosCaching/Public/Chaos/Adapters/CacheAdapter.h
@@ -7,6 +7,7 @@
 #include "Features/IModularFeature.h"
 #include "Templates/SubclassOf.h"
 #include "Chaos/PBDRigidsEvolutionFwd.h"
+#include "Chaos/Framework/PhysicsSolverBase.h"
 
 class UClass;
 class UChaosCache;
@@ -17,14 +18,10 @@ DECLARE_LOG_CATEGORY_EXTERN(LogCacheAdapter, Log, All);
 
 struct FPendingFrameWrite;
 
-namespace Chaos
-{
-	using FPhysicsSolver = FPBDRigidsSolver;
-}
-
 namespace Chaos
 {
 	class FComponentCacheAdapter;
+	class AChaosCacheManager;
 	struct FAdapterUtil
 	{
 		static CHAOSCACHING_API FComponentCacheAdapter* GetBestAdapterForClass(TSubclassOf<UPrimitiveComponent> InComponentClass, bool bAllowDerived = true);
@@ -106,13 +103,20 @@ namespace Chaos
 		virtual uint8 GetPriority() const = 0;
 
 		/**
-		 * Called to retrieve the solver for a specific component. Required until a more generic method
+		 * Called to retrieve the rigid solver for a specific component if it exists. Required until a more generic method
 		 * of solver binding for components is devised.
 		 * #BGTODO Remove when multiple solver concept moved into primitive component
 		 * @param InComponent Component to resolve the solver for
 		 */
 		virtual Chaos::FPhysicsSolver* GetComponentSolver(UPrimitiveComponent* InComponent) const = 0;
 
+		/**
+		 * Called to retrieve the base events solver for a specific component. If the component is
+		 * requiring it the solver could be rebuilt.
+		 * @param InComponent Component to resolve the solver for
+		 */
+		virtual Chaos::FPhysicsSolverEvents* BuildEventsSolver(UPrimitiveComponent* InComponent) const { return nullptr; }
+
 		/**
 		 * Called from the game thread to perform any global setup that the adapter may need to perform.
 		 */
diff --git a/Engine/Plugins/Experimental/ChaosCaching/Source/ChaosCaching/Public/Chaos/Adapters/GeometryCollectionComponentCacheAdapter.h b/Engine/Plugins/Experimental/ChaosCaching/Source/ChaosCaching/Public/Chaos/Adapters/GeometryCollectionComponentCacheAdapter.h
--- a/Engine/Plugins/Experimental/ChaosCaching/Source/ChaosCaching/Public/Chaos/Adapters/GeometryCollectionComponentCacheAdapter.h
+++ b/Engine/Plugins/Experimental/ChaosCaching/Source/ChaosCaching/Public/Chaos/Adapters/GeometryCollectionComponentCacheAdapter.h
@@ -234,17 +234,18 @@ namespace Chaos
 		virtual ~FGeometryCollectionCacheAdapter() = default;
 
 		// FComponentCacheAdapter interface
-		SupportType            SupportsComponentClass(UClass* InComponentClass) const override;
-		UClass*                GetDesiredClass() const override;
-		uint8                  GetPriority() const override;
-		FGuid                  GetGuid() const override;
-		Chaos::FPhysicsSolver* GetComponentSolver(UPrimitiveComponent* InComponent) const override;
-		bool                   ValidForPlayback(UPrimitiveComponent* InComponent, UChaosCache* InCache) const override;
-		void				   SetRestState(UPrimitiveComponent* InComponent, UChaosCache* InCache, const FTransform& InRootTransform, Chaos::FReal InTime) const override;
-		bool                   InitializeForRecord(UPrimitiveComponent* InComponent, UChaosCache* InCache) override;
-		bool                   InitializeForPlayback(UPrimitiveComponent* InComponent, UChaosCache* InCache, float InTime) override;
-		void                   Record_PostSolve(UPrimitiveComponent* InComp, const FTransform& InRootTransform, FPendingFrameWrite& OutFrame, Chaos::FReal InTime) const override;
-		void                   Playback_PreSolve(UPrimitiveComponent*								InComponent,
+		virtual SupportType            SupportsComponentClass(UClass* InComponentClass) const override;
+		virtual UClass*                GetDesiredClass() const override;
+		virtual uint8                  GetPriority() const override;
+		virtual FGuid                  GetGuid() const override;
+		virtual Chaos::FPhysicsSolver* GetComponentSolver(UPrimitiveComponent* InComponent) const override;
+		virtual Chaos::FPhysicsSolverEvents* BuildEventsSolver(UPrimitiveComponent* InComponent) const override;
+		virtual bool                   ValidForPlayback(UPrimitiveComponent* InComponent, UChaosCache* InCache) const override;
+		virtual void				   SetRestState(UPrimitiveComponent* InComponent, UChaosCache* InCache, const FTransform& InRootTransform, Chaos::FReal InTime) const override;
+		virtual bool                   InitializeForRecord(UPrimitiveComponent* InComponent, UChaosCache* InCache) override;
+		virtual bool                   InitializeForPlayback(UPrimitiveComponent* InComponent, UChaosCache* InCache, float InTime) override;
+		virtual void                   Record_PostSolve(UPrimitiveComponent* InComp, const FTransform& InRootTransform, FPendingFrameWrite& OutFrame, Chaos::FReal InTime) const override;
+		virtual void                   Playback_PreSolve(UPrimitiveComponent*								InComponent,
 												 UChaosCache*										InCache,
 												 Chaos::FReal                                       InTime,
 												 FPlaybackTickRecord&								TickRecord,
diff --git a/Engine/Plugins/Experimental/ChaosCaching/Source/ChaosCaching/Public/Chaos/Adapters/StaticMeshComponentCacheAdapter.h b/Engine/Plugins/Experimental/ChaosCaching/Source/ChaosCaching/Public/Chaos/Adapters/StaticMeshComponentCacheAdapter.h
--- a/Engine/Plugins/Experimental/ChaosCaching/Source/ChaosCaching/Public/Chaos/Adapters/StaticMeshComponentCacheAdapter.h
+++ b/Engine/Plugins/Experimental/ChaosCaching/Source/ChaosCaching/Public/Chaos/Adapters/StaticMeshComponentCacheAdapter.h
@@ -11,22 +11,22 @@ namespace Chaos
 	public:
 		virtual ~FStaticMeshCacheAdapter() = default;
 
-		SupportType            SupportsComponentClass(UClass* InComponentClass) const override;
-		UClass*                GetDesiredClass() const override;
-		uint8                  GetPriority() const override;
-		FGuid                  GetGuid() const override;
-		bool                   ValidForPlayback(UPrimitiveComponent* InComponent, UChaosCache* InCache) const override;
-		Chaos::FPhysicsSolver* GetComponentSolver(UPrimitiveComponent* InComponent) const override;
-		void				   SetRestState(UPrimitiveComponent* InComponent, UChaosCache* InCache, const FTransform& InRootTransform, Chaos::FReal InTime) const override;
-		bool                   InitializeForRecord(UPrimitiveComponent* InComponent, UChaosCache* InCache) override;
-		bool                   InitializeForPlayback(UPrimitiveComponent* InComponent, UChaosCache* InCache, float InTime) override;
-		void                   Record_PostSolve(UPrimitiveComponent* InComp, const FTransform& InRootTransform, FPendingFrameWrite& OutFrame, Chaos::FReal InTime) const override;
-		void                   Playback_PreSolve(UPrimitiveComponent* InComponent,
+		virtual SupportType            SupportsComponentClass(UClass* InComponentClass) const override;
+		virtual UClass*                GetDesiredClass() const override;
+		virtual uint8                  GetPriority() const override;
+		virtual FGuid                  GetGuid() const override;
+		virtual bool                   ValidForPlayback(UPrimitiveComponent* InComponent, UChaosCache* InCache) const override;
+		virtual Chaos::FPhysicsSolver* GetComponentSolver(UPrimitiveComponent* InComponent) const override;
+		virtual Chaos::FPhysicsSolverEvents* BuildEventsSolver(UPrimitiveComponent* InComponent) const override;
+		virtual void				   SetRestState(UPrimitiveComponent* InComponent, UChaosCache* InCache, const FTransform& InRootTransform, Chaos::FReal InTime) const override;
+		virtual bool                   InitializeForRecord(UPrimitiveComponent* InComponent, UChaosCache* InCache) override;
+		virtual bool                   InitializeForPlayback(UPrimitiveComponent* InComponent, UChaosCache* InCache, float InTime) override;
+		virtual void                   Record_PostSolve(UPrimitiveComponent* InComp, const FTransform& InRootTransform, FPendingFrameWrite& OutFrame, Chaos::FReal InTime) const override;
+		virtual void                   Playback_PreSolve(UPrimitiveComponent* InComponent,
 												 UChaosCache*										InCache,
 												 Chaos::FReal                                       InTime,
 												 FPlaybackTickRecord&								TickRecord,
 												 TArray<TPBDRigidParticleHandle<Chaos::FReal, 3>*>& OutUpdatedRigids) const override;
-
 	private:
 	};
 }    // namespace Chaos
diff --git a/Engine/Plugins/Experimental/ChaosCaching/Source/ChaosCaching/Public/Chaos/CacheCollection.h b/Engine/Plugins/Experimental/ChaosCaching/Source/ChaosCaching/Public/Chaos/CacheCollection.h
--- a/Engine/Plugins/Experimental/ChaosCaching/Source/ChaosCaching/Public/Chaos/CacheCollection.h
+++ b/Engine/Plugins/Experimental/ChaosCaching/Source/ChaosCaching/Public/Chaos/CacheCollection.h
@@ -21,7 +21,10 @@ public:
 
 	void FlushAllCacheWrites();
 
-	const TArray<UChaosCache*> GetCaches() const;
+	const TArray<UChaosCache*>& GetCaches() const {return Caches;}
+
+	/** Return the max duration of all the caches stored in the collection */
+	float GetMaxDuration() const;
 
 	UPROPERTY(EditAnywhere, Instanced, Category="Caching", meta=(EditFixedOrder))
 	TArray<TObjectPtr<UChaosCache>> Caches;
diff --git a/Engine/Plugins/Experimental/ChaosCaching/Source/ChaosCaching/Public/Chaos/CacheManagerActor.h b/Engine/Plugins/Experimental/ChaosCaching/Source/ChaosCaching/Public/Chaos/CacheManagerActor.h
--- a/Engine/Plugins/Experimental/ChaosCaching/Source/ChaosCaching/Public/Chaos/CacheManagerActor.h
+++ b/Engine/Plugins/Experimental/ChaosCaching/Source/ChaosCaching/Public/Chaos/CacheManagerActor.h
@@ -5,6 +5,7 @@
 #include "Chaos/ParticleHandleFwd.h"
 #include "Chaos/PBDRigidsEvolutionFwd.h"
 #include "ChaosCache.h"
+#include "Adapters/CacheAdapter.h"
 #include "Engine/EngineTypes.h"
 #include "Chaos/Core.h"
 #include "GameFramework/Actor.h"
@@ -39,7 +40,7 @@ namespace Chaos
 	using FPhysicsSolver = FPBDRigidsSolver;
 }
 
-USTRUCT()
+USTRUCT(BlueprintType)
 struct CHAOSCACHING_API FObservedComponent
 {
 	GENERATED_BODY()
@@ -52,17 +53,41 @@ struct CHAOSCACHING_API FObservedComponent
 	{
 	}
 
+	FObservedComponent& operator=(const FObservedComponent& OtherComponent)
+	{
+		CacheName = OtherComponent.CacheName;
+		bIsSimulating = OtherComponent.bIsSimulating;
+		bTriggered = OtherComponent.bTriggered;
+		AbsoluteTime = OtherComponent.AbsoluteTime;
+		TimeSinceTrigger = OtherComponent.TimeSinceTrigger;
+		Cache = OtherComponent.Cache;
+		TickRecord = OtherComponent.TickRecord;
+		BestFitAdapter = OtherComponent.BestFitAdapter;
+		SoftComponentRef = OtherComponent.SoftComponentRef;
+		
+		return *this;
+	}
+
 	/** Unique name for the cache, used as a key into the cache collection */
 	UPROPERTY(EditAnywhere, Category = "Caching")
 	FName CacheName;
 
+	/** Deprecated hard object reference. Not working with sequencer and take recorder since
+	 all the pointers from any other packages are cleared out. Use TSoftObjectPtr instead. */
+	UE_DEPRECATED(5.1, "This property is going to be deleted. Use the SoftComponentRef instead")
+	UPROPERTY()
+	FComponentReference ComponentRef;
+
 	/** The component observed by this object for either playback or recording */
 	UPROPERTY(EditAnywhere, Category = "Caching", meta = (UseComponentPicker, AllowAnyActor))
-	FComponentReference ComponentRef;
+	FSoftComponentReference SoftComponentRef;
 
 	/** Capture of the initial state of the component before cache manager takes control. */
 	UPROPERTY(EditAnywhere, Category = "Caching")
-	bool bIsSimulating;				  
+	bool bIsSimulating;
+
+	/** Post serialize function to transfer datas from the deprecated TObjectPtr -> TSoftObjectPtr */
+	void PostSerialize(const FArchive& Ar);
 
 	/** Prepare runtime tick data for a new run */
 	void ResetRuntimeData(const EStartMode ManagerStartMode);
@@ -140,6 +165,7 @@ public:
 	/** end AActor interface */
 
 	/** UObject interface */
+	virtual void Serialize(FArchive& Ar) override;
 #if WITH_EDITOR
 	friend class IChaosCachingEditorPlugin;
 #endif
@@ -179,24 +205,37 @@ public:
 	/** Returns true if this cache manager is allowed to record caches. */
 	bool CanRecord() const { return bCanRecord; }
 
+	/** Initialize the cache adapters before playing/recording the cache. */
+	void BeginEvaluate();
+	
+	/** Clean the cache adapters after playing/recording the cache.  */
+	void EndEvaluate();
+
+	/** Accessor to the manager observed components (read only) */
+	const TArray<FObservedComponent>& GetObservedComponents() const {return ObservedComponents;}
+
+	/** Accessor to the manager observed components (read/write) */
+	TArray<FObservedComponent>& GetObservedComponents() {return ObservedComponents;}
+
 protected:
+
 	/** AActor interface */
 	void BeginPlay() override;
 	void EndPlay(const EEndPlayReason::Type EndPlayReason) override;
 	/** End AActor interface */
-
+	
 	/**	Handles physics thread pre-solve (push kinematic data for components under playback) */
-	void HandlePreSolve(Chaos::FReal InDt, Chaos::FPhysicsSolver* InSolver);
-
-	/** Handles physics thread pre-buffer (mark dirty kinematic particles) */
-	void HandlePreBuffer(Chaos::FReal InDt, Chaos::FPhysicsSolver* InSolver);
-
-	/** Handles physics thread post-solve (record data for components under record) */
-	void HandlePostSolve(Chaos::FReal InDt, Chaos::FPhysicsSolver* InSolver);
-
-	/** Handles solver teardown due to solver destruction / stream-out */
-	void HandleTeardown(Chaos::FPhysicsSolver* InSolver);
-
+    void HandlePreSolve(Chaos::FReal InDt, Chaos::FPhysicsSolverEvents* InSolver);
+   
+    /** Handles physics thread pre-buffer (mark dirty kinematic particles) */
+    void HandlePreBuffer(Chaos::FReal InDt, Chaos::FPhysicsSolverEvents* InSolver);
+   
+    /** Handles physics thread post-solve (record data for components under record) */
+    void HandlePostSolve(Chaos::FReal InDt, Chaos::FPhysicsSolverEvents* InSolver);
+   
+    /** Handles solver teardown due to solver destruction / stream-out */
+    void HandleTeardown(Chaos::FPhysicsSolverEvents* InSolver);
+	
 	/** Evaluates and sets state for all observed components at the specified time. */
 	void OnStartFrameChanged(Chaos::FReal InT);
 
@@ -259,7 +298,7 @@ private:
 	TArray<Chaos::FComponentCacheAdapter*> ActiveAdapters;
 
 	/** List of particles returned by the adapter as requiring a kinematic update */
-	TMap<Chaos::FPhysicsSolver*, FPerSolverData> PerSolverData;
+	TMap<Chaos::FPhysicsSolverEvents*, FPerSolverData> PerSolverData;
 
 	/** Lists of currently open caches that need to be closed when complete */
 	TArray<TTuple<FCacheUserToken, UChaosCache*>> OpenRecordCaches;
diff --git a/Engine/Plugins/Experimental/ChaosCaching/Source/ChaosCaching/Public/Chaos/ChaosCache.h b/Engine/Plugins/Experimental/ChaosCaching/Source/ChaosCaching/Public/Chaos/ChaosCache.h
--- a/Engine/Plugins/Experimental/ChaosCaching/Source/ChaosCaching/Public/Chaos/ChaosCache.h
+++ b/Engine/Plugins/Experimental/ChaosCaching/Source/ChaosCaching/Public/Chaos/ChaosCache.h
@@ -171,6 +171,7 @@ struct FCacheEvaluationContext
 	bool                 bEvaluateCurves;
 	bool                 bEvaluateEvents;
 	TArray<int32>        EvaluationIndices;
+	bool                 bEvaluateChannels;
 };
 
 struct FCacheEvaluationResult
@@ -181,6 +182,7 @@ public:
 	TArray<FTransform>                     Transform;
 	TArray<TMap<FName, float>>             Curves;
 	TMap<FName, TArray<FCacheEventHandle>> Events;
+	TMap<FName, TArray<float>>             Channels;
 };
 
 struct FPendingParticleWrite
@@ -191,6 +193,15 @@ struct FPendingParticleWrite
 	TArray<TPair<FName, float>> PendingCurveData;
 };
 
+USTRUCT()
+struct FRichCurves
+{
+	GENERATED_BODY()
+	
+	UPROPERTY()
+	TArray<FRichCurve> RichCurves;
+};
+
 struct FPendingFrameWrite
 {
 	float                         Time;
@@ -198,6 +209,9 @@ struct FPendingFrameWrite
 	TArray<TPair<FName, float>>   PendingCurveData;
 	TMap<FName, FCacheEventTrack> PendingEvents;
 
+	TArray<int32>				  PendingChannelsIndices;
+	TMap<FName, TArray<float>>	  PendingChannelsData;
+
 	template<typename T>
 	FCacheEventTrack& FindOrAddEventTrack(FName InName)
 	{
@@ -368,6 +382,10 @@ public:
 	UPROPERTY()
 	TArray<FPerParticleCacheData> ParticleTracks;
 
+	/** Per-particle data, includes transforms, velocities and other per-particle, per-frame data */
+	UPROPERTY()
+	TMap<FName,FRichCurves> ChannelsTracks;
+
 	/** Per component/cache curve data, any continuous data that isn't per-particle can be stored here */
 	UPROPERTY()
 	TMap<FName, FRichCurve> CurveData;
diff --git a/Engine/Plugins/Experimental/ChaosCaching/Source/ChaosCaching/Public/Chaos/ChaosCachingPlugin.h b/Engine/Plugins/Experimental/ChaosCaching/Source/ChaosCaching/Public/Chaos/ChaosCachingPlugin.h
--- a/Engine/Plugins/Experimental/ChaosCaching/Source/ChaosCaching/Public/Chaos/ChaosCachingPlugin.h
+++ b/Engine/Plugins/Experimental/ChaosCaching/Source/ChaosCaching/Public/Chaos/ChaosCachingPlugin.h
@@ -10,6 +10,7 @@
 namespace Chaos {
 	class FGeometryCollectionCacheAdapter;
 	class FStaticMeshCacheAdapter;
+	class FSkeletalMeshCacheAdapter;
 }
 
 DECLARE_LOG_CATEGORY_EXTERN(LogChaosCache, Verbose, All);
@@ -45,4 +46,5 @@ private:
 	TUniquePtr<Chaos::FGeometryCollectionCacheAdapter> GeometryCollectionAdapter;
 	TUniquePtr<Chaos::FStaticMeshCacheAdapter> StaticMeshAdapter;
 
+	FDelegateHandle OnCreateMovieSceneObjectSpawnerHandle;
 };
diff --git a/Engine/Plugins/Experimental/ChaosCaching/Source/ChaosCaching/Public/Chaos/Sequencer/ChaosCacheObjectSpawner.h b/Engine/Plugins/Experimental/ChaosCaching/Source/ChaosCaching/Public/Chaos/Sequencer/ChaosCacheObjectSpawner.h
new file mode 100644
--- /dev/null
+++ b/Engine/Plugins/Experimental/ChaosCaching/Source/ChaosCaching/Public/Chaos/Sequencer/ChaosCacheObjectSpawner.h
@@ -0,0 +1,22 @@
+﻿// Copyright Epic Games, Inc. All Rights Reserved.
+
+#pragma once
+
+#include "LevelSequenceActorSpawner.h"
+
+/** Chaos cache Manager spawner to create a new cache manager from the spawnable template */
+class CHAOSCACHING_API FChaosCacheObjectSpawner : public FLevelSequenceActorSpawner
+{
+public:
+	FChaosCacheObjectSpawner();
+
+	/** Static method to create the object spawner */
+	static TSharedRef<IMovieSceneObjectSpawner> CreateObjectSpawner();
+
+	// IMovieSceneObjectSpawner interface
+	virtual UClass* GetSupportedTemplateType() const override;
+	virtual UObject* SpawnObject(FMovieSceneSpawnable& Spawnable, FMovieSceneSequenceIDRef TemplateID, IMovieScenePlayer& Player) override;
+	virtual void DestroySpawnedObject(UObject& Object) override;
+	virtual bool IsEditor() const override { return true; }
+	virtual int32 GetSpawnerPriority() const override { return 1; }
+};
\ No newline at end of file
diff --git a/Engine/Plugins/Experimental/ChaosCloth/ChaosCloth.uplugin b/Engine/Plugins/Experimental/ChaosCloth/ChaosCloth.uplugin
--- a/Engine/Plugins/Experimental/ChaosCloth/ChaosCloth.uplugin
+++ b/Engine/Plugins/Experimental/ChaosCloth/ChaosCloth.uplugin
@@ -20,5 +20,11 @@
         "Type": "Runtime",
         "LoadingPhase": "PostConfigInit"
       }
-    ]
+    ],
+	"Plugins": [
+		{
+			"Name": "ChaosCaching",
+			"Enabled": true
+		}
+	]
 }
diff --git a/Engine/Plugins/Experimental/ChaosCloth/Source/Chaos/ChaosCloth.Build.cs b/Engine/Plugins/Experimental/ChaosCloth/Source/Chaos/ChaosCloth.Build.cs
--- a/Engine/Plugins/Experimental/ChaosCloth/Source/Chaos/ChaosCloth.Build.cs
+++ b/Engine/Plugins/Experimental/ChaosCloth/Source/Chaos/ChaosCloth.Build.cs
@@ -16,7 +16,8 @@ public class ChaosCloth : ModuleRules
                 "ClothingSystemRuntimeCommon",
                 "ClothingSystemRuntimeInterface",
                 "Engine",
-                "Chaos"
+                "Chaos",
+                "ChaosCaching"
             }
 		);
 
diff --git a/Engine/Plugins/Experimental/ChaosCloth/Source/Chaos/Private/ChaosCloth/ChaosClothModule.cpp b/Engine/Plugins/Experimental/ChaosCloth/Source/Chaos/Private/ChaosCloth/ChaosClothModule.cpp
--- a/Engine/Plugins/Experimental/ChaosCloth/Source/Chaos/Private/ChaosCloth/ChaosClothModule.cpp
+++ b/Engine/Plugins/Experimental/ChaosCloth/Source/Chaos/Private/ChaosCloth/ChaosClothModule.cpp
@@ -3,6 +3,7 @@
 #include "ChaosCloth/ChaosClothModule.h"
 #include "ChaosCloth/ChaosClothPrivate.h"
 #include "ChaosCloth/ChaosClothingSimulationFactory.h"
+#include "ChaosCloth/SkeletalMeshComponentCacheAdapter.h"
 #include "Features/IModularFeatures.h"
 #include "Modules/ModuleManager.h"
 
@@ -16,17 +17,25 @@ class FChaosClothModule : public IChaosClothModuleInterface, public IClothingSim
     {
         check(GConfig);
 		IModularFeatures::Get().RegisterModularFeature(IClothingSimulationFactoryClassProvider::FeatureName, this);
+
+    	SkeletalMeshAdapter = MakeUnique<Chaos::FSkeletalMeshCacheAdapter>();
+    	Chaos::RegisterAdapter(SkeletalMeshAdapter.Get());
     }
 
     virtual void ShutdownModule() override
     {
 		IModularFeatures::Get().UnregisterModularFeature(IClothingSimulationFactoryClassProvider::FeatureName, this);
+
+    	UnregisterAdapter(SkeletalMeshAdapter.Get());
+    	SkeletalMeshAdapter = nullptr;
     }
 
 	TSubclassOf<UClothingSimulationFactory> GetClothingSimulationFactoryClass() const override
 	{
 		return UChaosClothingSimulationFactory::StaticClass();
 	}
+private:
+	TUniquePtr<Chaos::FSkeletalMeshCacheAdapter> SkeletalMeshAdapter;
 };
 
 //////////////////////////////////////////////////////////////////////////
diff --git a/Engine/Plugins/Experimental/ChaosCloth/Source/Chaos/Private/ChaosCloth/ChaosClothingSimulation.cpp b/Engine/Plugins/Experimental/ChaosCloth/Source/Chaos/Private/ChaosCloth/ChaosClothingSimulation.cpp
--- a/Engine/Plugins/Experimental/ChaosCloth/Source/Chaos/Private/ChaosCloth/ChaosClothingSimulation.cpp
+++ b/Engine/Plugins/Experimental/ChaosCloth/Source/Chaos/Private/ChaosCloth/ChaosClothingSimulation.cpp
@@ -579,7 +579,14 @@ void FClothingSimulation::Simulate(IClothingSimulationContext* InContext)
 		}
 
 		// Step the simulation
-		Solver->Update((Softs::FSolverReal)Context->DeltaSeconds);
+		if(Solver->GetEnableSolver() || (Context->CachedPositions.Num() == 0 && Context->CachedVelocities.Num() == 0))
+		{
+			Solver->Update((Softs::FSolverReal)Context->DeltaSeconds);
+		}
+		else
+		{
+			Solver->UpdateFromCache(Context->CachedPositions, Context->CachedVelocities);
+		}
 
 		// Keep the actual used number of iterations for the stats
 		NumIterations = Solver->GetNumUsedIterations();
diff --git a/Engine/Plugins/Experimental/ChaosCloth/Source/Chaos/Private/ChaosCloth/ChaosClothingSimulationSolver.cpp b/Engine/Plugins/Experimental/ChaosCloth/Source/Chaos/Private/ChaosCloth/ChaosClothingSimulationSolver.cpp
--- a/Engine/Plugins/Experimental/ChaosCloth/Source/Chaos/Private/ChaosCloth/ChaosClothingSimulationSolver.cpp
+++ b/Engine/Plugins/Experimental/ChaosCloth/Source/Chaos/Private/ChaosCloth/ChaosClothingSimulationSolver.cpp
@@ -97,6 +97,7 @@ FClothingSimulationSolver::FClothingSimulationSolver()
 	, WindVelocity(ClothingSimulationSolverDefault::WindVelocity)
 	, LegacyWindAdaption(0.f)
 	, bIsClothGravityOverrideEnabled(false)
+	, bEnableSolver(true)
 {
 	Softs::FSolverParticles LocalParticles;
 	Softs::FSolverRigidParticles RigidParticles;
@@ -947,6 +948,8 @@ void FClothingSimulationSolver::Update(Softs::FSolverReal InDeltaTime)
 		SCOPE_CYCLE_COUNTER(STAT_ChaosClothSolverUpdatePreSolverStep);
 
 		ApplyPreSimulationTransforms();
+		
+		EventPreSolve.Broadcast(DeltaTime);
 	}
 
 	// Advance Sim
@@ -955,21 +958,24 @@ void FClothingSimulationSolver::Update(Softs::FSolverReal InDeltaTime)
 		SCOPE_CYCLE_COUNTER(STAT_ChaosClothSolverUpdateSolverStep);
 		SCOPE_CYCLE_COUNTER(STAT_ClothInternalSolve);
 
-		// Update solver time dependent parameters
-		constexpr Softs::FSolverReal SolverFrequency = 60.f;  // TODO: This could become a solver property
+		if(bEnableSolver)
+		{
+			// Update solver time dependent parameters
+			constexpr Softs::FSolverReal SolverFrequency = 60.f;  // TODO: This could become a solver property
 
-		const int32 TimeDependentNumIterations = bClothSolverDisableTimeDependentNumIterations ?
-			NumIterations :
-			(int32)(SolverFrequency * DeltaTime * (Softs::FSolverReal)NumIterations);
+			const int32 TimeDependentNumIterations = bClothSolverDisableTimeDependentNumIterations ?
+				NumIterations : 
+				(int32)(SolverFrequency * DeltaTime * (Softs::FSolverReal)NumIterations);
 
-		Evolution->SetIterations(FMath::Clamp(TimeDependentNumIterations, 1, MaxNumIterations));
+			Evolution->SetIterations(FMath::Clamp(TimeDependentNumIterations, 1, MaxNumIterations));
 
-		// Advance substeps
-		const Softs::FSolverReal SubstepDeltaTime = DeltaTime / (Softs::FSolverReal)NumSubsteps;
+			// Advance substeps
+			const Softs::FSolverReal SubstepDeltaTime = DeltaTime / (Softs::FSolverReal)NumSubsteps;
 	
-		for (int32 i = 0; i < NumSubsteps; ++i)
-		{
-			Evolution->AdvanceOneTimeStep(SubstepDeltaTime, bClothSolverUseImprovedTimeStepSmoothing);
+			for (int32 i = 0; i < NumSubsteps; ++i)
+			{
+				Evolution->AdvanceOneTimeStep(SubstepDeltaTime, bClothSolverUseImprovedTimeStepSmoothing);
+			}
 		}
 
 		Time = Evolution->GetTime();
@@ -981,6 +987,8 @@ void FClothingSimulationSolver::Update(Softs::FSolverReal InDeltaTime)
 		TRACE_CPUPROFILER_EVENT_SCOPE(FClothingSimulationSolver_UpdatePostSolverStep);
 		SCOPE_CYCLE_COUNTER(STAT_ChaosClothSolverUpdatePostSolverStep);
 		SCOPE_CYCLE_COUNTER(STAT_ClothComputeNormals);
+		
+		EventPostSolve.Broadcast(DeltaTime);
 
 		PhysicsParallelFor(Cloths.Num(), [this](int32 ClothIndex)
 		{
@@ -993,6 +1001,25 @@ void FClothingSimulationSolver::Update(Softs::FSolverReal InDeltaTime)
 	OldLocalSpaceLocation = LocalSpaceLocation;
 }
 
+void FClothingSimulationSolver::UpdateFromCache(const TArray<FVector>& CachedPositions, const TArray<FVector>& CachedVelocities) 
+{
+	Chaos::Softs::FSolverParticles& SolverParticles = Evolution->Particles();
+	const int32 NumParticles = GetNumParticles();
+	if(CachedPositions.Num() == NumParticles)
+	{
+		for(int32 ParticleIndex = 0; ParticleIndex < NumParticles; ++ParticleIndex)
+		{
+			SolverParticles.X(ParticleIndex) = CachedPositions[ParticleIndex];
+			SolverParticles.V(ParticleIndex) = CachedVelocities[ParticleIndex];
+		}
+	}
+	PhysicsParallelFor(Cloths.Num(), [this](int32 ClothIndex)
+	{
+		FClothingSimulationCloth* const Cloth = Cloths[ClothIndex];
+		Cloth->PostUpdate(this);
+	}, /*bForceSingleThreaded =*/ !bClothSolverParallelClothPostUpdate);
+}
+
 int32 FClothingSimulationSolver::GetNumUsedIterations() const
 {
 	return Evolution->GetIterations();
@@ -1074,4 +1101,9 @@ FBoxSphereBounds FClothingSimulationSolver::CalculateBounds() const
 	return FBoxSphereBounds(LocalSpaceLocation, FVector(0.f), 0.f);
 }
 
+uint32 FClothingSimulationSolver::GetNumParticles() const
+{
+	return Evolution->Particles().Size();
+}
+
 } // End namespace Chaos
diff --git a/Engine/Plugins/Experimental/ChaosCloth/Source/Chaos/Private/ChaosCloth/SkeletalMeshComponentCacheAdapter.cpp b/Engine/Plugins/Experimental/ChaosCloth/Source/Chaos/Private/ChaosCloth/SkeletalMeshComponentCacheAdapter.cpp
new file mode 100644
--- /dev/null
+++ b/Engine/Plugins/Experimental/ChaosCloth/Source/Chaos/Private/ChaosCloth/SkeletalMeshComponentCacheAdapter.cpp
@@ -0,0 +1,286 @@
+﻿// Copyright Epic Games, Inc. All Rights Reserved.
+
+#include "ChaosCloth/SkeletalMeshComponentCacheAdapter.h"
+#include "Chaos/ParticleHandle.h"
+#include "Chaos/ChaosCache.h"
+#include "Chaos/CacheManagerActor.h"
+#include "Chaos/PBDEvolution.h"
+#include "ChaosCloth/ChaosClothingSimulation.h"
+#include "ChaosCloth/ChaosClothingSimulationSolver.h"
+#include "Components/SkeletalMeshComponent.h"
+
+namespace Chaos
+{
+	FComponentCacheAdapter::SupportType FSkeletalMeshCacheAdapter::SupportsComponentClass(UClass* InComponentClass) const
+	{
+		const UClass* Desired = GetDesiredClass();
+		if(InComponentClass == Desired)
+		{
+			return FComponentCacheAdapter::SupportType::Direct;
+		}
+		else if(InComponentClass->IsChildOf(Desired))
+		{
+			return FComponentCacheAdapter::SupportType::Derived;
+		}
+
+		return FComponentCacheAdapter::SupportType::None;
+	}
+
+	UClass* FSkeletalMeshCacheAdapter::GetDesiredClass() const
+	{
+		return USkeletalMeshComponent::StaticClass();
+	}
+
+	uint8 FSkeletalMeshCacheAdapter::GetPriority() const
+	{
+		return EngineAdapterPriorityBegin;
+	}
+
+	void FSkeletalMeshCacheAdapter::Record_PostSolve(UPrimitiveComponent* InComponent, const FTransform& InRootTransform, FPendingFrameWrite& OutFrame, Chaos::FReal InTime) const
+	{
+		if( Chaos::FClothingSimulationSolver* ClothSolver = GetClothSolver(InComponent))
+		{
+			const uint32 NumParticles = ClothSolver->GetNumParticles();
+			if(NumParticles > 0)
+			{
+				const Softs::FSolverVec3* ParticleXs = ClothSolver->GetParticleXs(0);
+				const Softs::FSolverVec3* ParticleVs = ClothSolver->GetParticleVs(0);
+				
+				TArray<float> PendingVX, PendingVY, PendingVZ, PendingPX, PendingPY, PendingPZ;
+				TArray<int32>& PendingID = OutFrame.PendingChannelsIndices;
+
+				PendingID.SetNum(NumParticles);
+				PendingVX.SetNum(NumParticles);
+				PendingVY.SetNum(NumParticles);
+				PendingVZ.SetNum(NumParticles);
+				PendingPX.SetNum(NumParticles);
+				PendingPY.SetNum(NumParticles);
+				PendingPZ.SetNum(NumParticles);
+
+				for(uint32 ParticleIndex = 0; ParticleIndex < NumParticles; ++ParticleIndex)
+				{
+					const Softs::FSolverVec3& ParticleV = ParticleVs[ParticleIndex];
+					const Softs::FSolverVec3& ParticleX = ParticleXs[ParticleIndex];
+					
+					// Adding the vertices relative position to the particle write datas
+					FPendingParticleWrite NewData;
+
+					NewData.ParticleIndex = ParticleIndex; 
+					OutFrame.PendingParticleData.Add(MoveTemp(NewData));
+
+					PendingID[ParticleIndex] = ParticleIndex;
+					PendingVX[ParticleIndex] = ParticleV.X;
+					PendingVY[ParticleIndex] = ParticleV.Y;
+					PendingVZ[ParticleIndex] = ParticleV.Z;
+
+					PendingPX[ParticleIndex] = ParticleX.X;
+					PendingPY[ParticleIndex] = ParticleX.Y;
+					PendingPZ[ParticleIndex] = ParticleX.Z;
+				}
+
+				OutFrame.PendingChannelsData.Add(VelocityXName, PendingVX);
+				OutFrame.PendingChannelsData.Add(VelocityYName, PendingVY);
+				OutFrame.PendingChannelsData.Add(VelocityZName, PendingVZ);
+				OutFrame.PendingChannelsData.Add(PositionXName, PendingPX);
+				OutFrame.PendingChannelsData.Add(PositionYName, PendingPY);
+				OutFrame.PendingChannelsData.Add(PositionZName, PendingPZ);
+			}
+		}
+	}
+
+	void FSkeletalMeshCacheAdapter::Playback_PreSolve(UPrimitiveComponent* InComponent, UChaosCache* InCache, Chaos::FReal InTime, FPlaybackTickRecord& TickRecord, TArray<TPBDRigidParticleHandle<Chaos::FReal, 3>*>& OutUpdatedRigids) const
+	{
+		if( Chaos::FClothingSimulationSolver* ClothSolver = GetClothSolver(InComponent))
+		{
+			FCacheEvaluationContext Context(TickRecord);
+			Context.bEvaluateTransform = false;
+			Context.bEvaluateCurves = true;
+			Context.bEvaluateEvents = false;
+
+			// The evaluated result are already in world space since we are passing the tickrecord.spacetransform in the context
+			FCacheEvaluationResult EvaluatedResult = InCache->Evaluate(Context, nullptr);
+			const int32 NumCurves = EvaluatedResult.Curves.Num();
+
+			const uint32 NumParticles = ClothSolver->GetNumParticles();
+			if(NumParticles > 0 && NumParticles == NumCurves)
+			{
+				Softs::FSolverVec3* ParticleXs = ClothSolver->GetParticleXs(0);
+				Softs::FSolverVec3* ParticleVs = ClothSolver->GetParticleVs(0);
+
+				TArray<float>* PendingVX = EvaluatedResult.Channels.Find(VelocityXName);
+				TArray<float>* PendingVY = EvaluatedResult.Channels.Find(VelocityYName);
+				TArray<float>* PendingVZ = EvaluatedResult.Channels.Find(VelocityZName);
+				TArray<float>* PendingPX = EvaluatedResult.Channels.Find(PositionXName);
+				TArray<float>* PendingPY = EvaluatedResult.Channels.Find(PositionYName);
+				TArray<float>* PendingPZ = EvaluatedResult.Channels.Find(PositionZName);
+					
+				if(PendingVX && PendingVY && PendingVZ && PendingPZ && PendingPY && PendingPZ)
+				{
+					for(uint32 ParticleIndex = 0; ParticleIndex < NumParticles; ++ParticleIndex)
+					{
+						Softs::FSolverVec3& ParticleV = ParticleVs[ParticleIndex];
+						Softs::FSolverVec3& ParticleX = ParticleXs[ParticleIndex];
+					
+						// Directly set the result of the cache into the solver particles
+						const int32 TransformIndex = EvaluatedResult.ParticleIndices[ParticleIndex];
+
+						ParticleV.X = (*PendingVX)[TransformIndex];
+						ParticleV.Y = (*PendingVY)[TransformIndex];
+						ParticleV.Z = (*PendingVZ)[TransformIndex];
+						ParticleX.X = (*PendingPX)[TransformIndex];
+						ParticleX.Y = (*PendingPY)[TransformIndex];
+						ParticleX.Z = (*PendingPZ)[TransformIndex];
+					}
+				}
+			}
+		}
+	}
+
+	FGuid FSkeletalMeshCacheAdapter::GetGuid() const
+	{
+		FGuid NewGuid;
+		checkSlow(FGuid::Parse(TEXT("4328BED8C52D4E07BE0CA2433940AF6D"), NewGuid));
+		return NewGuid;
+	}
+
+	bool FSkeletalMeshCacheAdapter::ValidForPlayback(UPrimitiveComponent* InComponent, UChaosCache* InCache) const
+	{
+		// If we have a skel mesh we can play back any cache as long as it has one or more tracks
+		const USkeletalMeshComponent* MeshComp = CastChecked<USkeletalMeshComponent>(InComponent);
+		return MeshComp && MeshComp->GetSkinnedAsset() && InCache->TrackToParticle.Num() > 0;
+	}
+
+	Chaos::FClothingSimulationSolver* FSkeletalMeshCacheAdapter::GetClothSolver(UPrimitiveComponent* InComponent) const
+	{
+		if(InComponent)
+		{
+			if( USkeletalMeshComponent* MeshComp = CastChecked<USkeletalMeshComponent>(InComponent) )
+			{
+				if(MeshComp->GetClothingSimulation())
+				{
+					return StaticCast<Chaos::FClothingSimulation*>(MeshComp->GetClothingSimulation())->Solver.Get();
+				}
+			}
+		}
+		return nullptr;
+	}
+
+	Chaos::FPhysicsSolverEvents* FSkeletalMeshCacheAdapter::BuildEventsSolver(UPrimitiveComponent* InComponent) const
+	{
+		if(USkeletalMeshComponent* MeshComp = CastChecked<USkeletalMeshComponent>(InComponent))
+		{
+			// Initialize the physics solver at the beginning of the play/record
+			MeshComp->RecreatePhysicsState();
+			MeshComp->RecreateClothingActors();
+
+			return GetClothSolver(InComponent);
+		}
+		return nullptr;
+	}
+	
+	Chaos::FPhysicsSolver* FSkeletalMeshCacheAdapter::GetComponentSolver(UPrimitiveComponent* InComponent) const
+	{
+		return nullptr;
+	}
+	
+	void FSkeletalMeshCacheAdapter::SetRestState(UPrimitiveComponent* InComponent, UChaosCache* InCache, const FTransform& InRootTransform, Chaos::FReal InTime) const
+	{
+		if (!InCache || InCache->GetDuration() == 0.0f)
+		{
+			return;
+		}
+
+		if(USkeletalMeshComponent* MeshComp = CastChecked<USkeletalMeshComponent>(InComponent))
+		{
+			InitializeForRestState(InComponent);
+		
+			FPlaybackTickRecord TickRecord;
+			TickRecord.SetLastTime(InTime);
+		
+			FCacheEvaluationContext Context(TickRecord);
+			Context.bEvaluateTransform = false;
+			Context.bEvaluateCurves = false;
+			Context.bEvaluateEvents = false;
+			Context.bEvaluateChannels = true;
+
+			FCacheEvaluationResult EvaluatedResult = InCache->Evaluate(Context,nullptr);
+
+			if( MeshComp->GetClothingSimulationContext())
+			{
+				FClothingSimulationContext* SimulationContext =
+						StaticCast<FClothingSimulationContext*>(MeshComp->GetClothingSimulationContext());
+				
+				const int32 NumCurves = EvaluatedResult.ParticleIndices.Num();
+				if(NumCurves == EvaluatedResult.ParticleIndices.Num())
+				{
+					SimulationContext->CachedPositions.SetNum(NumCurves);
+					SimulationContext->CachedVelocities.SetNum(NumCurves);
+
+					TArray<float>* PendingVX = EvaluatedResult.Channels.Find(VelocityXName);
+					TArray<float>* PendingVY = EvaluatedResult.Channels.Find(VelocityYName);
+					TArray<float>* PendingVZ = EvaluatedResult.Channels.Find(VelocityZName);
+					TArray<float>* PendingPX = EvaluatedResult.Channels.Find(PositionXName);
+					TArray<float>* PendingPY = EvaluatedResult.Channels.Find(PositionYName);
+					TArray<float>* PendingPZ = EvaluatedResult.Channels.Find(PositionZName);
+					
+					if(PendingVX && PendingVY && PendingVZ && PendingPZ && PendingPY && PendingPZ)
+					{
+						for(int32 ParticleIndex = 0; ParticleIndex < NumCurves; ++ParticleIndex)
+						{
+							const int32 TransformIndex = EvaluatedResult.ParticleIndices[ParticleIndex];
+							
+							SimulationContext->CachedVelocities[ParticleIndex].X = (*PendingVX)[TransformIndex];
+							SimulationContext->CachedVelocities[ParticleIndex].Y = (*PendingVY)[TransformIndex];
+							SimulationContext->CachedVelocities[ParticleIndex].Z = (*PendingVZ)[TransformIndex];
+							SimulationContext->CachedPositions[ParticleIndex].X = (*PendingPX)[TransformIndex];
+							SimulationContext->CachedPositions[ParticleIndex].Y = (*PendingPY)[TransformIndex];
+							SimulationContext->CachedPositions[ParticleIndex].Z = (*PendingPZ)[TransformIndex];
+						}
+					}
+				}
+			}
+		}
+	}
+
+	void FSkeletalMeshCacheAdapter::InitializeForRestState(UPrimitiveComponent* InComponent) const
+	{
+		if(USkeletalMeshComponent* MeshComp = CastChecked<USkeletalMeshComponent>(InComponent))
+		{
+#if WITH_EDITOR
+			if(!MeshComp->GetUpdateClothInEditor() || !MeshComp->GetUpdateAnimationInEditor())
+			{
+				MeshComp->SetUpdateAnimationInEditor(true);
+				MeshComp->SetUpdateClothInEditor(true);
+			}
+#endif
+			// Initialize the physics solver if the cloth/animation in editor are not set
+			MeshComp->RecreatePhysicsState();
+			MeshComp->RecreateClothingActors();
+		}
+
+		if( Chaos::FClothingSimulationSolver* ClothSolver = GetClothSolver(InComponent))
+		{
+			ClothSolver->SetEnableSolver(false);
+		}
+	}
+
+	bool FSkeletalMeshCacheAdapter::InitializeForRecord(UPrimitiveComponent* InComponent, UChaosCache* InCache)
+	{
+		if( Chaos::FClothingSimulationSolver* ClothSolver = GetClothSolver(InComponent))
+		{
+			ClothSolver->SetEnableSolver(true);
+		}
+		return true;
+	}
+
+	bool FSkeletalMeshCacheAdapter::InitializeForPlayback(UPrimitiveComponent* InComponent, UChaosCache* InCache, float InTime)
+	{
+		EnsureIsInGameThreadContext();
+		
+		if( Chaos::FClothingSimulationSolver* ClothSolver = GetClothSolver(InComponent))
+		{
+			ClothSolver->SetEnableSolver(false);
+		}
+		return true;
+	}
+}    // namespace Chaos
diff --git a/Engine/Plugins/Experimental/ChaosCloth/Source/Chaos/Private/ChaosCloth/SkeletalMeshComponentCacheAdapter.h b/Engine/Plugins/Experimental/ChaosCloth/Source/Chaos/Private/ChaosCloth/SkeletalMeshComponentCacheAdapter.h
new file mode 100644
--- /dev/null
+++ b/Engine/Plugins/Experimental/ChaosCloth/Source/Chaos/Private/ChaosCloth/SkeletalMeshComponentCacheAdapter.h
@@ -0,0 +1,59 @@
+﻿// Copyright Epic Games, Inc. All Rights Reserved.
+
+#pragma once
+
+#include "Chaos/Adapters/CacheAdapter.h"
+
+namespace Chaos
+{
+	class FClothingSimulationSolver;
+	
+	/**
+	 * Skeletal mesh cache adapter to be able to cache cloth simulation datas through the chaos cache system
+	 */
+	class FSkeletalMeshCacheAdapter : public FComponentCacheAdapter
+	{
+	public:
+		
+		virtual ~FSkeletalMeshCacheAdapter() = default;
+
+		// ~Begin FComponentCacheAdapter interface
+		virtual SupportType            SupportsComponentClass(UClass* InComponentClass) const override;
+		virtual UClass*                GetDesiredClass() const override;
+		virtual uint8                  GetPriority() const override;
+		virtual FGuid                  GetGuid() const override;
+		virtual bool                   ValidForPlayback(UPrimitiveComponent* InComponent, UChaosCache* InCache) const override;
+		virtual Chaos::FPhysicsSolver* GetComponentSolver(UPrimitiveComponent* InComponent) const override;
+		virtual Chaos::FPhysicsSolverEvents* BuildEventsSolver(UPrimitiveComponent* InComponent) const override;
+		virtual void				   SetRestState(UPrimitiveComponent* InComponent, UChaosCache* InCache, const FTransform& InRootTransform, Chaos::FReal InTime) const override;
+		virtual bool                   InitializeForRecord(UPrimitiveComponent* InComponent, UChaosCache* InCache) override;
+		virtual bool                   InitializeForPlayback(UPrimitiveComponent* InComponent, UChaosCache* InCache, float InTime) override;
+		virtual void                   Record_PostSolve(UPrimitiveComponent* InComp, const FTransform& InRootTransform, FPendingFrameWrite& OutFrame, Chaos::FReal InTime) const override;
+		virtual void                   Playback_PreSolve(UPrimitiveComponent* InComponent,
+										 UChaosCache*										InCache,
+										 Chaos::FReal                                       InTime,
+										 FPlaybackTickRecord&								TickRecord,
+										 TArray<TPBDRigidParticleHandle<Chaos::FReal, 3>*>& OutUpdatedRigids) const override;
+		// ~End FComponentCacheAdapter interface
+
+	private :
+
+		/** Return the adapter component cloth solver
+		 * @param InComponent skelmesh component from which the cloth solver will be extracted
+		 * @return Cloth solver used by the skelmesh component
+		 */
+		Chaos::FClothingSimulationSolver* GetClothSolver(UPrimitiveComponent* InComponent) const;
+
+		/** Initialize the skelmesh adapter when changing the rest state 
+		* @param InComponent skelmesh component from which the cloth solver will be extracted
+		*/
+		void InitializeForRestState(UPrimitiveComponent* InComponent) const;
+
+		inline static const FName VelocityXName = TEXT("VelocityX");
+		inline static const FName VelocityYName = TEXT("VelocityY");
+		inline static const FName VelocityZName = TEXT("VelocityZ");
+		inline static const FName PositionXName = TEXT("PositionX");
+		inline static const FName PositionYName = TEXT("PositionY");
+		inline static const FName PositionZName = TEXT("PositionZ");
+	};
+}    // namespace Chaos
diff --git a/Engine/Plugins/Experimental/ChaosCloth/Source/Chaos/Public/ChaosCloth/ChaosClothingSimulation.h b/Engine/Plugins/Experimental/ChaosCloth/Source/Chaos/Public/ChaosCloth/ChaosClothingSimulation.h
--- a/Engine/Plugins/Experimental/ChaosCloth/Source/Chaos/Public/ChaosCloth/ChaosClothingSimulation.h
+++ b/Engine/Plugins/Experimental/ChaosCloth/Source/Chaos/Public/ChaosCloth/ChaosClothingSimulation.h
@@ -22,6 +22,7 @@ namespace Chaos
 	class FClothingSimulationMesh;
 	class FClothingSimulationCloth;
 	class FClothingSimulationCollider;
+	class FSkeletalMeshCacheAdapter;
 
 	typedef FClothingSimulationContextCommon FClothingSimulationContext;
 
@@ -34,6 +35,8 @@ namespace Chaos
 		FClothingSimulation();
 		virtual ~FClothingSimulation() override;
 
+		friend FSkeletalMeshCacheAdapter;
+
 	protected:
 		// IClothingSimulation interface
 		virtual void Initialize() override;
diff --git a/Engine/Plugins/Experimental/ChaosCloth/Source/Chaos/Public/ChaosCloth/ChaosClothingSimulationSolver.h b/Engine/Plugins/Experimental/ChaosCloth/Source/Chaos/Public/ChaosCloth/ChaosClothingSimulationSolver.h
--- a/Engine/Plugins/Experimental/ChaosCloth/Source/Chaos/Public/ChaosCloth/ChaosClothingSimulationSolver.h
+++ b/Engine/Plugins/Experimental/ChaosCloth/Source/Chaos/Public/ChaosCloth/ChaosClothingSimulationSolver.h
@@ -5,6 +5,7 @@
 #include "Chaos/PBDSoftsEvolutionFwd.h"
 #include "Chaos/ArrayCollectionArray.h"
 #include "Chaos/Transform.h"
+#include "Chaos/Framework/PhysicsSolverBase.h"
 #include "Chaos/ImplicitObject.h"
 #include "PhysicsProxy/PerSolverFieldSystem.h"
 
@@ -20,11 +21,12 @@ namespace Chaos
 	class FClothingSimulationMesh;
 
 	// Solver simulation node
-	class FClothingSimulationSolver final
+	class CHAOSCLOTH_API FClothingSimulationSolver final : public FPhysicsSolverEvents
 	{
 	public:
 		FClothingSimulationSolver();
 		~FClothingSimulationSolver();
+		
 
 		// ---- Animatable property setters ----
 		void SetLocalSpaceLocation(const FVec3& InLocalSpaceLocation, bool bReset = false);
@@ -50,6 +52,8 @@ namespace Chaos
 		int32 GetMaxNumIterations() const { return MaxNumIterations; }
 		void SetNumSubsteps(int32 InNumSubsteps) { NumSubsteps = InNumSubsteps; }
 		int32 GetNumSubsteps() const { return NumSubsteps; }
+		void SetEnableSolver(bool InbEnableSolver) { bEnableSolver = InbEnableSolver; }
+		bool GetEnableSolver() const { return bEnableSolver; }
 		// ---- End of the animatable property setters ----
 
 		// ---- Object management functions ----
@@ -64,8 +68,12 @@ namespace Chaos
 		TConstArrayView<const FClothingSimulationCloth*> GetCloths() const { return Cloths; }
 
 		// Update solver properties before simulation.
+		
 		void Update(FSolverReal InDeltaTime);
 
+		/** Set the cached positions onto the particles */
+		void UpdateFromCache(const TArray<FVector>& CachedPositions, const TArray<FVector>& CachedVelocities);
+
 		// Return the actual of number of iterations used by the Evolution solver after the update (different from the number of iterations, depends on frame rate)
 		int32 GetNumUsedIterations() const;
 
@@ -152,6 +160,7 @@ namespace Chaos
 		const Softs::FSolverReal* GetParticleInvMasses(int32 Offset) const;
 		const FClothConstraints& GetClothConstraints(int32 Offset) const { return *ClothsConstraints.FindChecked(Offset); }
 		FClothConstraints& GetClothConstraints(int32 Offset) { return *ClothsConstraints.FindChecked(Offset); }
+		uint32 GetNumParticles() const;
 		// ---- End of the Cloth interface ----
 
 		// ---- Collider interface ----
@@ -244,6 +253,9 @@ namespace Chaos
 
 		// Field system unique to the cloth solver
 		FPerSolverFieldSystem PerSolverField;
+
+		// Boolean to enable/disable solver in case caching is used
+		bool bEnableSolver;
 	};
 
 } // namespace Chaos
diff --git a/Engine/Source/Runtime/ClothingSystemRuntimeCommon/Public/ClothingSimulation.h b/Engine/Source/Runtime/ClothingSystemRuntimeCommon/Public/ClothingSimulation.h
--- a/Engine/Source/Runtime/ClothingSystemRuntimeCommon/Public/ClothingSimulation.h
+++ b/Engine/Source/Runtime/ClothingSystemRuntimeCommon/Public/ClothingSimulation.h
@@ -71,6 +71,12 @@ public:
 
 	// The predicted LOD of the skeletal mesh component running the simulation
 	int32 PredictedLod;
+	
+	// World space cached positions for the kinematics targets
+	TArray<FVector> CachedPositions;
+
+	// World space cached velocities for the kinematics targets
+	TArray<FVector> CachedVelocities;
 };
 
 // Base simulation to fill in common data for the base context
diff --git a/Engine/Source/Runtime/Engine/Classes/Components/SkeletalMeshComponent.h b/Engine/Source/Runtime/Engine/Classes/Components/SkeletalMeshComponent.h
--- a/Engine/Source/Runtime/Engine/Classes/Components/SkeletalMeshComponent.h
+++ b/Engine/Source/Runtime/Engine/Classes/Components/SkeletalMeshComponent.h
@@ -1455,10 +1455,16 @@ public:
 
 	/** Get the current clothing simulation (read only) */
 	const IClothingSimulation* GetClothingSimulation() const;
+	
+	/** Get the current clothing simulation (read/write) */
+	IClothingSimulation* GetClothingSimulation();
 
 	/** Get the current clothing simulation context (read only) */
 	const IClothingSimulationContext* GetClothingSimulationContext() const;
-
+	
+	/** Get the current clothing simulation context (read/write) */
+	IClothingSimulationContext* GetClothingSimulationContext();
+	
 	/** Get the current interactor for a clothing simulation, if the current simulation supports runtime interaction. */
 	UFUNCTION(BlueprintCallable, Category=Clothing)
 	UClothingSimulationInteractor* GetClothingSimulationInteractor() const;
diff --git a/Engine/Source/Runtime/Engine/Classes/Engine/EngineTypes.h b/Engine/Source/Runtime/Engine/Classes/Engine/EngineTypes.h
--- a/Engine/Source/Runtime/Engine/Classes/Engine/EngineTypes.h
+++ b/Engine/Source/Runtime/Engine/Classes/Engine/EngineTypes.h
@@ -3114,22 +3114,14 @@ struct FConstrainComponentPropName
 };
 
 /** 
- *	Struct that allows for different ways to reference a component. 
- *	If just an Actor is specified, will return RootComponent of that Actor.
+ *	Base class for the hard/soft component reference structs 
  */
 USTRUCT(BlueprintType)
-struct ENGINE_API FComponentReference
+struct ENGINE_API FBaseComponentReference
 {
 	GENERATED_BODY()
 
-	FComponentReference() : OtherActor(nullptr) {}
-
-	/** 
-	 * Pointer to a different Actor that owns the Component.  
-	 * If this is not provided the reference refers to a component on this / the same actor.
-	 */
-	UPROPERTY(EditInstanceOnly, BlueprintReadWrite, Category=Component, meta = (DisplayName = "Referenced Actor"))
-	TObjectPtr<AActor> OtherActor;
+	FBaseComponentReference()  {}
 
 	/** Name of component to use. If this is not specified the reference refers to the root component. */
 	UPROPERTY(EditAnywhere, BlueprintReadWrite, Category=Component, meta = (DisplayName = "Component Name"))
@@ -3142,15 +3134,73 @@ struct ENGINE_API FComponentReference
 	/** Allows direct setting of first component to constraint. */
 	TWeakObjectPtr<class UActorComponent> OverrideComponent;
 
+	/** Extract the actual component pointer from this reference given a search actor */
+	class UActorComponent* ExtractComponent(AActor* SearchActor) const;
+
+	/** FBaseComponentReference == operator */
+	bool operator== (const FBaseComponentReference& Other) const
+	{
+		return ComponentProperty == Other.ComponentProperty && PathToComponent == Other.PathToComponent && OverrideComponent == Other.OverrideComponent;
+	}
+};
+
+/** 
+ *	Struct that allows for different ways to reference a component using TObjectPtr. 
+ *	If just an Actor is specified, will return RootComponent of that Actor.
+ */
+USTRUCT(BlueprintType)
+struct ENGINE_API FComponentReference : public FBaseComponentReference
+{
+	GENERATED_BODY()
+
+	FComponentReference() : OtherActor(nullptr) {}
+
+	/** 
+	 * Hard Pointer to a different Actor that owns the Component.  
+	 * If this is not provided the reference refers to a component on this / the same actor.
+	 */
+	UPROPERTY(EditInstanceOnly, BlueprintReadWrite, Category=Component, meta = (DisplayName = "Referenced Actor"))
+	TObjectPtr<AActor> OtherActor;
+
 	/** Get the actual component pointer from this reference */
 	class UActorComponent* GetComponent(AActor* OwningActor) const;
 
+	/** FComponentReference == operator */
 	bool operator== (const FComponentReference& Other) const
 	{
-		return OtherActor == Other.OtherActor && ComponentProperty == Other.ComponentProperty && PathToComponent == Other.PathToComponent && OverrideComponent == Other.OverrideComponent;
+		return (OtherActor == Other.OtherActor) && (FBaseComponentReference::operator==(Other));
+	}
+};
+
+/** 
+ *	Struct that allows for different ways to reference a component using TSoftObjectPtr. 
+ *	If just an Actor is specified, will return RootComponent of that Actor.
+ */
+USTRUCT(BlueprintType)
+struct ENGINE_API FSoftComponentReference : public FBaseComponentReference
+{
+	GENERATED_BODY()
+
+	FSoftComponentReference() : OtherActor(nullptr) {}
+
+	/** 
+	 * Soft Pointer to a different Actor that owns the Component.  
+	 * If this is not provided the reference refers to a component on this / the same actor.
+	 */
+	UPROPERTY(EditInstanceOnly, BlueprintReadWrite, Category=Component, meta = (DisplayName = "Referenced Actor"))
+	TSoftObjectPtr<AActor> OtherActor;
+
+	/** Get the actual component pointer from this reference */
+	class UActorComponent* GetComponent(AActor* OwningActor) const;
+
+	/** FSoftComponentReference == operator */
+	bool operator== (const FSoftComponentReference& Other) const
+	{
+		return (OtherActor == Other.OtherActor) && (FBaseComponentReference::operator==(Other));
 	}
 };
 
+
 /** Types of valid physical material mask colors which may be associated with a physical material */
 UENUM(BlueprintType)
 namespace EPhysicalMaterialMaskColor
diff --git a/Engine/Source/Runtime/Engine/Private/Components/SkeletalMeshComponent.cpp b/Engine/Source/Runtime/Engine/Private/Components/SkeletalMeshComponent.cpp
--- a/Engine/Source/Runtime/Engine/Private/Components/SkeletalMeshComponent.cpp
+++ b/Engine/Source/Runtime/Engine/Private/Components/SkeletalMeshComponent.cpp
@@ -2080,13 +2080,25 @@ const IClothingSimulation* USkeletalMeshComponent::GetClothingSimulation() const
 	return ClothingSimulation;
 }
 
+IClothingSimulation* USkeletalMeshComponent::GetClothingSimulation()
+{
+	return ClothingSimulation;
+}
+
 const IClothingSimulationContext* USkeletalMeshComponent::GetClothingSimulationContext() const
 {
 	return ClothingSimulationContext;
 }
 
+IClothingSimulationContext* USkeletalMeshComponent::GetClothingSimulationContext()
+{
+	check(IsInGameThread());
+	return ClothingSimulationContext;
+}
+
 UClothingSimulationInteractor* USkeletalMeshComponent::GetClothingSimulationInteractor() const
 {
+	check(IsInGameThread());
 	return ClothingInteractor;
 }
 
diff --git a/Engine/Source/Runtime/Engine/Private/EngineTypes.cpp b/Engine/Source/Runtime/Engine/Private/EngineTypes.cpp
--- a/Engine/Source/Runtime/Engine/Private/EngineTypes.cpp
+++ b/Engine/Source/Runtime/Engine/Private/EngineTypes.cpp
@@ -120,7 +120,8 @@ FLightmassDebugOptions::FLightmassDebugOptions()
 	, ExecutionTimeDivisor(15.0f)
 {}
 
-UActorComponent* FComponentReference::GetComponent(AActor* OwningActor) const
+
+UActorComponent* FBaseComponentReference::ExtractComponent(AActor* SearchActor) const 
 {
 	UActorComponent* Result = nullptr;
 
@@ -131,8 +132,6 @@ UActorComponent* FComponentReference::GetComponent(AActor* OwningActor) const
 	}
 	else
 	{
-		// Look in Actor if specified, OwningActor if not
-		AActor* SearchActor = (OtherActor != NULL) ? ToRawPtr(OtherActor) : OwningActor;
 		if(SearchActor)
 		{
 			if(ComponentProperty != NAME_None)
@@ -158,6 +157,18 @@ UActorComponent* FComponentReference::GetComponent(AActor* OwningActor) const
 	return Result;
 }
 
+UActorComponent* FComponentReference::GetComponent(AActor* OwningActor) const
+{
+	AActor* SearchActor = (OtherActor != NULL) ? ToRawPtr(OtherActor) : OwningActor;
+	return ExtractComponent(SearchActor);
+}
+
+UActorComponent* FSoftComponentReference::GetComponent(AActor* OwningActor) const
+{
+	AActor* SearchActor = (OtherActor.IsValid()) ? OtherActor.Get() : OwningActor;
+	return ExtractComponent(SearchActor);
+}
+
 const TCHAR* LexToString(const EWorldType::Type Value)
 {
 	switch (Value)
diff --git a/Engine/Source/Runtime/Experimental/Chaos/Private/Chaos/Framework/PhysicsSolverBase.cpp b/Engine/Source/Runtime/Experimental/Chaos/Private/Chaos/Framework/PhysicsSolverBase.cpp
--- a/Engine/Source/Runtime/Experimental/Chaos/Private/Chaos/Framework/PhysicsSolverBase.cpp
+++ b/Engine/Source/Runtime/Experimental/Chaos/Private/Chaos/Framework/PhysicsSolverBase.cpp
@@ -19,42 +19,42 @@ namespace Chaos
 		BufferMode = InBufferMode;
 	}
 
-	FDelegateHandle FPhysicsSolverBase::AddPreAdvanceCallback(FSolverPreAdvance::FDelegate InDelegate)
+	FDelegateHandle FPhysicsSolverEvents::AddPreAdvanceCallback(FSolverPreAdvance::FDelegate InDelegate)
 	{
 		return EventPreSolve.Add(InDelegate);
 	}
 
-	bool FPhysicsSolverBase::RemovePreAdvanceCallback(FDelegateHandle InHandle)
+	bool FPhysicsSolverEvents::RemovePreAdvanceCallback(FDelegateHandle InHandle)
 	{
 		return EventPreSolve.Remove(InHandle);
 	}
 
-	FDelegateHandle FPhysicsSolverBase::AddPreBufferCallback(FSolverPreBuffer::FDelegate InDelegate)
+	FDelegateHandle FPhysicsSolverEvents::AddPreBufferCallback(FSolverPreBuffer::FDelegate InDelegate)
 	{
 		return EventPreBuffer.Add(InDelegate);
 	}
 
-	bool FPhysicsSolverBase::RemovePreBufferCallback(FDelegateHandle InHandle)
+	bool FPhysicsSolverEvents::RemovePreBufferCallback(FDelegateHandle InHandle)
 	{
 		return EventPreBuffer.Remove(InHandle);
 	}
 
-	FDelegateHandle FPhysicsSolverBase::AddPostAdvanceCallback(FSolverPostAdvance::FDelegate InDelegate)
+	FDelegateHandle FPhysicsSolverEvents::AddPostAdvanceCallback(FSolverPostAdvance::FDelegate InDelegate)
 	{
 		return EventPostSolve.Add(InDelegate);
 	}
 
-	bool FPhysicsSolverBase::RemovePostAdvanceCallback(FDelegateHandle InHandle)
+	bool FPhysicsSolverEvents::RemovePostAdvanceCallback(FDelegateHandle InHandle)
 	{
 		return EventPostSolve.Remove(InHandle);
 	}
 
-	FDelegateHandle FPhysicsSolverBase::AddTeardownCallback(FSolverTeardown::FDelegate InDelegate)
+	FDelegateHandle FPhysicsSolverEvents::AddTeardownCallback(FSolverTeardown::FDelegate InDelegate)
 	{
 		return EventTeardown.Add(InDelegate);
 	}
 
-	bool FPhysicsSolverBase::RemoveTeardownCallback(FDelegateHandle InHandle)
+	bool FPhysicsSolverEvents::RemoveTeardownCallback(FDelegateHandle InHandle)
 	{
 		return EventTeardown.Remove(InHandle);
 	}
diff --git a/Engine/Source/Runtime/Experimental/Chaos/Public/Chaos/Framework/PhysicsSolverBase.h b/Engine/Source/Runtime/Experimental/Chaos/Public/Chaos/Framework/PhysicsSolverBase.h
--- a/Engine/Source/Runtime/Experimental/Chaos/Public/Chaos/Framework/PhysicsSolverBase.h
+++ b/Engine/Source/Runtime/Experimental/Chaos/Public/Chaos/Framework/PhysicsSolverBase.h
@@ -29,7 +29,6 @@ namespace Chaos
 	class IRewindCallback;
 
 	extern CHAOS_API int32 GSingleThreadedPhysics;
-
 	extern CHAOS_API int32 UseAsyncInterpolation;
 	extern CHAOS_API int32 ForceDisableAsyncPhysics;
 	extern CHAOS_API FRealSingle AsyncInterpolationMultiplier;
@@ -166,8 +165,54 @@ namespace Chaos
 		TaskGraph,
 		SingleThread
 	};
+	
+	/** Base solver class storing events that will be used by the derived solver during the solve.
+	 * Currently used by caching but could be used for any other pre/post solve works that need to be done.  */
+	class CHAOS_API FPhysicsSolverEvents
+	{
+
+	public:
+
+		/** Only allow construction with valid parameters as well as restricting to module construction */
+		virtual ~FPhysicsSolverEvents() { ClearCallbacks(); }
+		
+		/** Events */
+		/** Pre advance is called before any physics processing or simulation happens in a given physics update */
+		FDelegateHandle AddPreAdvanceCallback(FSolverPreAdvance::FDelegate InDelegate);
+		bool            RemovePreAdvanceCallback(FDelegateHandle InHandle);
+
+		/** Pre buffer happens after the simulation has been advanced (particle positions etc. will have been updated) but GT results haven't been prepared yet */
+		FDelegateHandle AddPreBufferCallback(FSolverPreAdvance::FDelegate InDelegate);
+		bool            RemovePreBufferCallback(FDelegateHandle InHandle);
+
+		/** Post advance happens after all processing and results generation has been completed */
+		FDelegateHandle AddPostAdvanceCallback(FSolverPostAdvance::FDelegate InDelegate);
+		bool            RemovePostAdvanceCallback(FDelegateHandle InHandle);
+
+		/** Teardown happens as the solver is destroyed or streamed out */
+		FDelegateHandle AddTeardownCallback(FSolverTeardown::FDelegate InDelegate);
+		bool            RemoveTeardownCallback(FDelegateHandle InHandle);
 
-	class CHAOS_API FPhysicsSolverBase
+		/** Clear all the callbacks*/
+		void            ClearCallbacks() 
+		{
+			EventPostSolve.Clear(); 
+			EventPreBuffer.Clear(); 
+			EventPostSolve.Clear(); 
+			EventTeardown.Clear();
+		}
+	
+	protected:
+			
+		/** Storage for events, see the Add/Remove pairs above for event timings */
+		FSolverPreAdvance EventPreSolve;
+		FSolverPreBuffer EventPreBuffer;
+		FSolverPostAdvance EventPostSolve;
+		FSolverTeardown EventTeardown;
+	};
+
+
+	class CHAOS_API FPhysicsSolverBase : public FPhysicsSolverEvents
 	{
 	public:
 
@@ -601,33 +646,12 @@ namespace Chaos
 		TArray<FGeometryParticle*> UniqueIdxToGTParticles;
 
 	public:
-		/** Events */
-		/** Pre advance is called before any physics processing or simulation happens in a given physics update */
-		FDelegateHandle AddPreAdvanceCallback(FSolverPreAdvance::FDelegate InDelegate);
-		bool            RemovePreAdvanceCallback(FDelegateHandle InHandle);
-
-		/** Pre buffer happens after the simulation has been advanced (particle positions etc. will have been updated) but GT results haven't been prepared yet */
-		FDelegateHandle AddPreBufferCallback(FSolverPreAdvance::FDelegate InDelegate);
-		bool            RemovePreBufferCallback(FDelegateHandle InHandle);
-
-		/** Post advance happens after all processing and results generation has been completed */
-		FDelegateHandle AddPostAdvanceCallback(FSolverPostAdvance::FDelegate InDelegate);
-		bool            RemovePostAdvanceCallback(FDelegateHandle InHandle);
-
-		/** Teardown happens as the solver is destroyed or streamed out */
-		FDelegateHandle AddTeardownCallback(FSolverTeardown::FDelegate InDelegate);
-		bool            RemoveTeardownCallback(FDelegateHandle InHandle);
 
 		/** Get the lock used for external data manipulation. A better API would be to use scoped locks so that getting a write lock is non-const */
 		//NOTE: this is a const operation so that you can acquire a read lock on a const solver. The assumption is that non-const write operations are already marked non-const
 		FPhysSceneLock& GetExternalDataLock_External() const { return *ExternalDataLock_External; }
 
 	protected:
-		/** Storage for events, see the Add/Remove pairs above for event timings */
-		FSolverPreAdvance EventPreSolve;
-		FSolverPreBuffer EventPreBuffer;
-		FSolverPostAdvance EventPostSolve;
-		FSolverTeardown EventTeardown;
 
 		void TrackGTParticle_External(FGeometryParticle& Particle);
 		void ClearGTParticle_External(FGeometryParticle& Particle);
diff --git a/Engine/Source/Runtime/MovieScene/Public/IMovieSceneObjectSpawner.h b/Engine/Source/Runtime/MovieScene/Public/IMovieSceneObjectSpawner.h
--- a/Engine/Source/Runtime/MovieScene/Public/IMovieSceneObjectSpawner.h
+++ b/Engine/Source/Runtime/MovieScene/Public/IMovieSceneObjectSpawner.h
@@ -66,6 +66,12 @@ public:
 	  */
 	virtual bool IsEditor() const { return false; }
 
+	/** 
+	  * @return Spawner priority in order to sort the list of the spawners. If several spawners 
+	  *	match the template type, the one with the highest priority will be picked. 
+	  */
+	virtual int32 GetSpawnerPriority() const { return 0; }
+
 #if WITH_EDITOR
 	
 	/**
