Chaos Cache Manager tracks

diff --git a/Engine/Plugins/Experimental/ChaosCaching/Source/ChaosCaching/Private/Chaos/Sequencer/MovieSceneChaosCacheSection.cpp b/Engine/Plugins/Experimental/ChaosCaching/Source/ChaosCaching/Private/Chaos/Sequencer/MovieSceneChaosCacheSection.cpp
new file mode 100644
--- /dev/null
+++ b/Engine/Plugins/Experimental/ChaosCaching/Source/ChaosCaching/Private/Chaos/Sequencer/MovieSceneChaosCacheSection.cpp
@@ -0,0 +1,32 @@
+﻿// Copyright Epic Games, Inc. All Rights Reserved.
+
+#include "Chaos/Sequencer/MovieSceneChaosCacheSection.h"
+#include "Chaos/CacheCollection.h"
+#include "MovieScene.h"
+
+FMovieSceneChaosCacheParams::FMovieSceneChaosCacheParams() : FMovieSceneBaseCacheParams()
+{
+	CacheCollection = nullptr;
+}
+
+float FMovieSceneChaosCacheParams::GetSequenceLength() const
+{
+	if(CacheCollection)
+	{
+		return CacheCollection->GetMaxDuration();
+	}
+	return 0.f;
+}
+
+UMovieSceneChaosCacheSection::UMovieSceneChaosCacheSection(const FObjectInitializer& ObjectInitializer)
+	: Super(ObjectInitializer)
+{
+	ParamsPtr = &Params;
+#if WITH_EDITOR
+	InitializePlayRate();
+#endif
+}
+
+
+
+
diff --git a/Engine/Plugins/Experimental/ChaosCaching/Source/ChaosCaching/Private/Chaos/Sequencer/MovieSceneChaosCacheTemplate.cpp b/Engine/Plugins/Experimental/ChaosCaching/Source/ChaosCaching/Private/Chaos/Sequencer/MovieSceneChaosCacheTemplate.cpp
new file mode 100644
--- /dev/null
+++ b/Engine/Plugins/Experimental/ChaosCaching/Source/ChaosCaching/Private/Chaos/Sequencer/MovieSceneChaosCacheTemplate.cpp
@@ -0,0 +1,114 @@
+﻿// Copyright Epic Games, Inc. All Rights Reserved.
+
+#include "MovieSceneChaosCacheTemplate.h"
+
+#include "Chaos/CacheCollection.h"
+#include "Chaos/CacheManagerActor.h"
+
+DECLARE_CYCLE_STAT(TEXT("Chaos Cache Evaluate"), MovieSceneEval_ChaosCache_Evaluate, STATGROUP_MovieSceneEval);
+DECLARE_CYCLE_STAT(TEXT("Chaos Cache Token Execute"), MovieSceneEval_ChaosCache_TokenExecute, STATGROUP_MovieSceneEval);
+
+/** Used to set Manual Tick back to previous when outside section */
+struct FPreAnimatedChaosCacheTokenProducer : IMovieScenePreAnimatedTokenProducer
+{
+	virtual IMovieScenePreAnimatedTokenPtr CacheExistingState(UObject& Object) const
+	{
+		struct FToken : IMovieScenePreAnimatedToken
+		{
+			FToken(AChaosCacheManager* InChaosCache)
+			{
+				// Cache this object's current update flag and animation mode
+				StartMode = InChaosCache->StartMode;
+				CacheMode = InChaosCache->CacheMode;
+			}
+
+			virtual void RestoreState(UObject& ObjectToRestore, const UE::MovieScene::FRestoreStateParams& Params)
+			{
+				AChaosCacheManager* ChaosCache = CastChecked<AChaosCacheManager>(&ObjectToRestore);
+				ChaosCache->StartMode = StartMode;
+				ChaosCache->CacheMode = CacheMode;
+			}
+			EStartMode StartMode;
+			ECacheMode CacheMode;
+			
+		};
+
+		return FToken(CastChecked<AChaosCacheManager>(&Object));
+	}
+	static FMovieSceneAnimTypeID GetAnimTypeID()
+	{
+		return TMovieSceneAnimTypeID<FPreAnimatedChaosCacheTokenProducer>();
+	}
+};
+
+/** A movie scene execution token that executes a ChaosCache */
+struct FChaosCacheExecutionToken : IMovieSceneExecutionToken
+{
+	FChaosCacheExecutionToken(const FMovieSceneChaosCacheSectionTemplateParameters &InParams) :
+		Params(InParams)
+	{}
+
+	static AChaosCacheManager* ChaosCacheFromObject(UObject* BoundObject)
+	{
+		if (AActor* Actor = Cast<AActor>(BoundObject))
+		{
+			if (AChaosCacheManager* ChaosCache = Cast<AChaosCacheManager>(Actor))
+			{
+				return ChaosCache;
+			}
+		}
+		return nullptr;
+	}
+
+	/** Execute this token, operating on all objects referenced by 'Operand' */
+	virtual void Execute(const FMovieSceneContext& Context, const FMovieSceneEvaluationOperand& Operand, FPersistentEvaluationData& PersistentData, IMovieScenePlayer& Player) override
+	{
+		MOVIESCENE_DETAILED_SCOPE_CYCLE_COUNTER(MovieSceneEval_ChaosCache_TokenExecute)
+		if (Operand.ObjectBindingID.IsValid())
+		{
+			for (TWeakObjectPtr<> WeakObj : Player.FindBoundObjects(Operand))
+			{
+				if (UObject* Obj = WeakObj.Get())
+				{
+					AChaosCacheManager*  ChaosCache = ChaosCacheFromObject(Obj);
+					if (ChaosCache)
+					{ 
+						// Set the ChaosCache on the component only if it's set and valid in the Params
+						if (Params.ChaosCacheParams.CacheCollection && Params.ChaosCacheParams.CacheCollection != ChaosCache->CacheCollection)
+						{
+							ChaosCache->CacheCollection = Params.ChaosCacheParams.CacheCollection;
+						}
+						else
+						{
+							// It could be unset if the Params was referencing a transient ChaosCache
+							// In that case, use the ChaosCache that is set on the component
+							Params.ChaosCacheParams.CacheCollection = ChaosCache->CacheCollection;
+						}
+						Player.SavePreAnimatedState(*ChaosCache, FPreAnimatedChaosCacheTokenProducer::GetAnimTypeID(), FPreAnimatedChaosCacheTokenProducer());
+						ChaosCache->StartMode = EStartMode::Timed;
+						ChaosCache->CacheMode = ECacheMode::None;
+						
+						// calculate the time at which to evaluate the animation
+						ChaosCache->SetStartTime( Params.MapTimeToAnimation(Params.ChaosCacheParams,
+							ChaosCache->CacheCollection->GetMaxDuration(), Context.GetTime(), Context.GetFrameRate()));
+					}
+				}
+			}
+		}
+	}
+
+	FMovieSceneChaosCacheSectionTemplateParameters Params;
+};
+
+FMovieSceneChaosCacheSectionTemplate::FMovieSceneChaosCacheSectionTemplate(const UMovieSceneChaosCacheSection& InSection)
+	: Params(const_cast<FMovieSceneChaosCacheParams&> (InSection.Params), InSection.GetInclusiveStartFrame(), InSection.GetExclusiveEndFrame())
+{
+}
+
+//We use a token here so we can set the manual tick state back to what it was previously when outside this section.
+//This is similar to how Skeletal Animation evaluation also works.
+void FMovieSceneChaosCacheSectionTemplate::Evaluate(const FMovieSceneEvaluationOperand& Operand, const FMovieSceneContext& Context, const FPersistentEvaluationData& PersistentData, FMovieSceneExecutionTokens& ExecutionTokens) const
+{
+	MOVIESCENE_DETAILED_SCOPE_CYCLE_COUNTER(MovieSceneEval_ChaosCache_Evaluate)
+	ExecutionTokens.Add(FChaosCacheExecutionToken(Params));
+}
\ No newline at end of file
diff --git a/Engine/Plugins/Experimental/ChaosCaching/Source/ChaosCaching/Private/Chaos/Sequencer/MovieSceneChaosCacheTemplate.h b/Engine/Plugins/Experimental/ChaosCaching/Source/ChaosCaching/Private/Chaos/Sequencer/MovieSceneChaosCacheTemplate.h
new file mode 100644
--- /dev/null
+++ b/Engine/Plugins/Experimental/ChaosCaching/Source/ChaosCaching/Private/Chaos/Sequencer/MovieSceneChaosCacheTemplate.h
@@ -0,0 +1,38 @@
+﻿// Copyright Epic Games, Inc. All Rights Reserved.
+
+#pragma once
+
+#include "Evaluation/MovieSceneBaseCacheTemplate.h"
+#include "Chaos/Sequencer/MovieSceneChaosCacheSection.h"
+#include "MovieSceneChaosCacheTemplate.generated.h"
+
+USTRUCT()
+struct FMovieSceneChaosCacheSectionTemplateParameters : public FMovieSceneBaseCacheSectionTemplateParameters
+{
+	GENERATED_BODY()
+
+	FMovieSceneChaosCacheSectionTemplateParameters() {}
+
+	FMovieSceneChaosCacheSectionTemplateParameters(const FMovieSceneChaosCacheParams& InChaosCacheParams, FFrameNumber InSectionStartTime, FFrameNumber InSectionEndTime)
+		: FMovieSceneBaseCacheSectionTemplateParameters(InSectionStartTime, InSectionEndTime), ChaosCacheParams(InChaosCacheParams)
+ 
+	{}
+
+	UPROPERTY()
+	FMovieSceneChaosCacheParams ChaosCacheParams;
+};
+
+USTRUCT()
+struct FMovieSceneChaosCacheSectionTemplate : public FMovieSceneEvalTemplate
+{
+	GENERATED_BODY()
+
+	FMovieSceneChaosCacheSectionTemplate() {}
+	FMovieSceneChaosCacheSectionTemplate(const UMovieSceneChaosCacheSection& Section);
+
+	virtual UScriptStruct& GetScriptStructImpl() const override { return *StaticStruct(); }
+	virtual void Evaluate(const FMovieSceneEvaluationOperand& Operand, const FMovieSceneContext& Context, const FPersistentEvaluationData& PersistentData, FMovieSceneExecutionTokens& ExecutionTokens) const override;
+
+	UPROPERTY()
+	FMovieSceneChaosCacheSectionTemplateParameters Params;
+};
diff --git a/Engine/Plugins/Experimental/ChaosCaching/Source/ChaosCaching/Private/Chaos/Sequencer/MovieSceneChaosCacheTrack.cpp b/Engine/Plugins/Experimental/ChaosCaching/Source/ChaosCaching/Private/Chaos/Sequencer/MovieSceneChaosCacheTrack.cpp
new file mode 100644
--- /dev/null
+++ b/Engine/Plugins/Experimental/ChaosCaching/Source/ChaosCaching/Private/Chaos/Sequencer/MovieSceneChaosCacheTrack.cpp
@@ -0,0 +1,123 @@
+﻿// Copyright Epic Games, Inc. All Rights Reserved.
+
+#include "Chaos/Sequencer/MovieSceneChaosCacheTrack.h"
+#include "Chaos/CacheManagerActor.h"
+#include "Chaos/CacheCollection.h"
+#include "MovieScene.h"
+#include "Chaos/Sequencer/MovieSceneChaosCacheSection.h"
+#include "MovieSceneChaosCacheTemplate.h"
+
+#define LOCTEXT_NAMESPACE "MovieSceneChaosCacheTrack"
+
+/* UMovieSceneChaosCacheTrack constructor
+ *****************************************************************************/
+
+UMovieSceneChaosCacheTrack::UMovieSceneChaosCacheTrack(const FObjectInitializer& ObjectInitializer)
+	: Super(ObjectInitializer)
+{
+#if WITH_EDITORONLY_DATA
+	TrackTint = FColor(96, 128, 24, 65);
+#endif
+
+	SupportedBlendTypes.Add(EMovieSceneBlendType::Absolute);
+
+	EvalOptions.bCanEvaluateNearestSection = true;
+	EvalOptions.bEvaluateInPreroll = true;
+}
+
+/* UMovieSceneChaosCacheTrack interface
+ *****************************************************************************/
+
+UMovieSceneSection* UMovieSceneChaosCacheTrack::AddNewAnimation(FFrameNumber KeyTime, AChaosCacheManager* ChaosCache)
+{
+	UMovieSceneChaosCacheSection* NewSection = Cast<UMovieSceneChaosCacheSection>(CreateNewSection());
+	{
+		const FFrameTime AnimationLength = ChaosCache->CacheCollection->GetMaxDuration() * GetTypedOuter<UMovieScene>()->GetTickResolution();
+		const int32 IFrameNumber = AnimationLength.FrameNumber.Value + (int)(AnimationLength.GetSubFrame() + 0.5f) + 1;
+		NewSection->InitialPlacementOnRow(AnimationSections, KeyTime, IFrameNumber, INDEX_NONE);
+		
+		NewSection->Params.CacheCollection = ChaosCache->CacheCollection;
+	}
+
+	AddSection(*NewSection);
+
+	return NewSection;
+}
+
+TArray<UMovieSceneSection*> UMovieSceneChaosCacheTrack::GetAnimSectionsAtTime(FFrameNumber Time)
+{
+	TArray<UMovieSceneSection*> Sections;
+	for (auto Section : AnimationSections)
+	{
+		if (Section->IsTimeWithinSection(Time))
+		{
+			Sections.Add(Section);
+		}
+	}
+
+	return Sections;
+}
+
+FMovieSceneEvalTemplatePtr UMovieSceneChaosCacheTrack::CreateTemplateForSection(const UMovieSceneSection& InSection) const
+{
+	return FMovieSceneChaosCacheSectionTemplate(*CastChecked<UMovieSceneChaosCacheSection>(&InSection));
+}
+
+/* UMovieSceneTrack interface
+ *****************************************************************************/
+
+const TArray<UMovieSceneSection*>& UMovieSceneChaosCacheTrack::GetAllSections() const
+{
+	return AnimationSections;
+}
+
+bool UMovieSceneChaosCacheTrack::SupportsType(TSubclassOf<UMovieSceneSection> SectionClass) const
+{
+	return SectionClass == UMovieSceneChaosCacheSection::StaticClass();
+}
+
+UMovieSceneSection* UMovieSceneChaosCacheTrack::CreateNewSection()
+{
+	return NewObject<UMovieSceneChaosCacheSection>(this, NAME_None, RF_Transactional);
+}
+
+void UMovieSceneChaosCacheTrack::RemoveAllAnimationData()
+{
+	AnimationSections.Empty();
+}
+
+bool UMovieSceneChaosCacheTrack::HasSection(const UMovieSceneSection& Section) const
+{
+	return AnimationSections.Contains(&Section);
+}
+
+void UMovieSceneChaosCacheTrack::AddSection(UMovieSceneSection& Section)
+{
+	AnimationSections.Add(&Section);
+}
+
+void UMovieSceneChaosCacheTrack::RemoveSection(UMovieSceneSection& Section)
+{
+	AnimationSections.Remove(&Section);
+}
+
+void UMovieSceneChaosCacheTrack::RemoveSectionAt(int32 SectionIndex)
+{
+	AnimationSections.RemoveAt(SectionIndex);
+}
+
+bool UMovieSceneChaosCacheTrack::IsEmpty() const
+{
+	return AnimationSections.Num() == 0;
+}
+
+#if WITH_EDITORONLY_DATA
+
+FText UMovieSceneChaosCacheTrack::GetDefaultDisplayName() const
+{
+	return LOCTEXT("TrackName", "Chaos Cache");
+}
+
+#endif
+
+#undef LOCTEXT_NAMESPACE
diff --git a/Engine/Plugins/Experimental/ChaosCaching/Source/ChaosCaching/Public/Chaos/Sequencer/MovieSceneChaosCacheSection.h b/Engine/Plugins/Experimental/ChaosCaching/Source/ChaosCaching/Public/Chaos/Sequencer/MovieSceneChaosCacheSection.h
new file mode 100644
--- /dev/null
+++ b/Engine/Plugins/Experimental/ChaosCaching/Source/ChaosCaching/Public/Chaos/Sequencer/MovieSceneChaosCacheSection.h
@@ -0,0 +1,36 @@
+﻿// Copyright Epic Games, Inc. All Rights Reserved.
+
+#pragma once
+
+#include "Sections/MovieSceneBaseCacheSection.h"
+#include "MovieSceneChaosCacheSection.generated.h"
+
+USTRUCT()
+struct CHAOSCACHING_API FMovieSceneChaosCacheParams : public FMovieSceneBaseCacheParams
+{
+	GENERATED_BODY()
+
+	FMovieSceneChaosCacheParams();
+	virtual ~FMovieSceneChaosCacheParams() override {}
+
+	/** Gets the animation sequence length, not modified by play rate */
+	virtual float GetSequenceLength() const override;
+	
+	/** The animation this section plays */
+	UPROPERTY(EditAnywhere, Category = "ChaosCache")
+	TObjectPtr<class UChaosCacheCollection> CacheCollection;
+};
+
+/**
+ * Movie scene section that control ChaosCache playback
+ */
+UCLASS(MinimalAPI)
+class UMovieSceneChaosCacheSection : public UMovieSceneBaseCacheSection
+{
+	GENERATED_UCLASS_BODY()
+
+public:
+
+	UPROPERTY(EditAnywhere, Category = "Animation", meta = (ShowOnlyInnerProperties))
+	FMovieSceneChaosCacheParams Params;
+};
diff --git a/Engine/Plugins/Experimental/ChaosCaching/Source/ChaosCaching/Public/Chaos/Sequencer/MovieSceneChaosCacheTrack.h b/Engine/Plugins/Experimental/ChaosCaching/Source/ChaosCaching/Public/Chaos/Sequencer/MovieSceneChaosCacheTrack.h
new file mode 100644
--- /dev/null
+++ b/Engine/Plugins/Experimental/ChaosCaching/Source/ChaosCaching/Public/Chaos/Sequencer/MovieSceneChaosCacheTrack.h
@@ -0,0 +1,53 @@
+﻿// Copyright Epic Games, Inc. All Rights Reserved.
+
+#pragma once
+
+#include "MovieSceneNameableTrack.h"
+#include "Chaos/CacheManagerActor.h"
+#include "Compilation/IMovieSceneTrackTemplateProducer.h"
+#include "MovieSceneChaosCacheTrack.generated.h"
+
+/**
+ * Handles animation of ChaosCache
+ */
+UCLASS(MinimalAPI)
+class UMovieSceneChaosCacheTrack
+	: public UMovieSceneNameableTrack
+	, public IMovieSceneTrackTemplateProducer
+{
+	GENERATED_UCLASS_BODY()
+
+public:
+
+	/** Adds a new animation to this track */
+	virtual UMovieSceneSection* AddNewAnimation(FFrameNumber KeyTime, class AChaosCacheManager* ChaosCache);
+
+	/** Gets the animation sections at a certain time */
+	TArray<UMovieSceneSection*> GetAnimSectionsAtTime(FFrameNumber Time);
+
+public:
+
+	// UMovieSceneTrack interface
+	virtual void RemoveAllAnimationData() override;
+	virtual bool HasSection(const UMovieSceneSection& Section) const override;
+	virtual void AddSection(UMovieSceneSection& Section) override;
+	virtual void RemoveSection(UMovieSceneSection& Section) override;
+	virtual void RemoveSectionAt(int32 SectionIndex) override;
+	virtual bool IsEmpty() const override;
+	virtual const TArray<UMovieSceneSection*>& GetAllSections() const override;
+	virtual bool SupportsType(TSubclassOf<UMovieSceneSection> SectionClass) const override;
+	virtual UMovieSceneSection* CreateNewSection() override;
+
+	// ~IMovieSceneTrackTemplateProducer interface
+	virtual FMovieSceneEvalTemplatePtr CreateTemplateForSection(const UMovieSceneSection& InSection) const override;
+
+#if WITH_EDITORONLY_DATA
+	virtual FText GetDefaultDisplayName() const override;
+#endif
+
+private:
+
+	/** List of all animation sections */
+	UPROPERTY()
+	TArray<TObjectPtr<UMovieSceneSection>> AnimationSections;
+};
diff --git a/Engine/Plugins/Experimental/ChaosCaching/Source/ChaosCachingEditor/ChaosCachingEditor.Build.cs b/Engine/Plugins/Experimental/ChaosCaching/Source/ChaosCachingEditor/ChaosCachingEditor.Build.cs
--- a/Engine/Plugins/Experimental/ChaosCaching/Source/ChaosCachingEditor/ChaosCachingEditor.Build.cs
+++ b/Engine/Plugins/Experimental/ChaosCaching/Source/ChaosCachingEditor/ChaosCachingEditor.Build.cs
@@ -25,7 +25,18 @@ namespace UnrealBuildTool.Rules
 					"PhysicsCore",
 					"ChaosCaching",
 					"GeometryCollectionEngine",
-					"LevelEditor"
+					"LevelEditor",
+					"LevelSequence",
+                    "MovieScene",
+                    "MovieSceneTools",
+                    "MovieSceneTracks",
+                    "Sequencer",
+                    "TimeManagement",
+					"EditorFramework",
+					"TakesCore",
+					"TakeRecorder",
+					"SceneOutliner",
+					"TakeTrackRecorders",
 				});
 
 			SetupModulePhysicsSupport(Target);
diff --git a/Engine/Plugins/Experimental/ChaosCaching/Source/ChaosCachingEditor/Private/Chaos/ChaosCachingEditorPlugin.cpp b/Engine/Plugins/Experimental/ChaosCaching/Source/ChaosCachingEditor/Private/Chaos/ChaosCachingEditorPlugin.cpp
--- a/Engine/Plugins/Experimental/ChaosCaching/Source/ChaosCachingEditor/Private/Chaos/ChaosCachingEditorPlugin.cpp
+++ b/Engine/Plugins/Experimental/ChaosCaching/Source/ChaosCachingEditor/Private/Chaos/ChaosCachingEditorPlugin.cpp
@@ -11,18 +11,166 @@
 #include "Chaos/CacheManagerCustomization.h"
 #include "Chaos/CacheCollectionFactory.h"
 #include "Chaos/CacheCollection.h"
+#include "Sequencer/TakeRecorderChaosCacheSource.h"
 #include "CoreMinimal.h"
 #include "Engine/Selection.h"
 #include "Kismet2/ComponentEditorUtils.h"
 #include "Modules/ModuleManager.h"
 #include "PropertyEditorModule.h"
-#include "ToolMenus.h"
 #include "LevelEditor.h"
+#include "ISequencerModule.h"
+#include "Sequencer/ChaosCacheTrackEditor.h"
+#include "ITakeRecorderModule.h"
+#include "Framework/Application/SlateApplication.h"
+#include "Framework/MultiBox/MultiBoxBuilder.h"
+#include "SceneOutlinerPublicTypes.h"
+#include "ActorTreeItem.h"
+#include "SceneOutlinerModule.h"
 
 IMPLEMENT_MODULE(IChaosCachingEditorPlugin, ChaosCachingEditor)
 
 #define LOCTEXT_NAMESPACE "CacheEditorPlugin"
 
+static void AddChaosCacheSource(UTakeRecorderSources* Sources,  const FName& SubjectName)
+{
+	FScopedTransaction Transaction(LOCTEXT("AddChaosCacheSource","Add Chaos Cache Source"));
+
+	Sources->Modify();
+
+	UTakeRecorderChaosCacheSource* NewSource = Sources->AddSource<UTakeRecorderChaosCacheSource>();
+}
+
+void AddChaosCacheSources(
+	UTakeRecorderSources* TakeRecorderSources,
+	TArrayView<AActor* const> ActorsToRecord)
+{
+	if (ActorsToRecord.Num() > 0)
+	{
+		FScopedTransaction Transaction(FText::Format(
+			LOCTEXT("AddSources", "Add Recording {0}|plural(one=Source, other=Sources)"), ActorsToRecord.Num()));
+		TakeRecorderSources->Modify();
+
+		for (AActor* Actor : ActorsToRecord)
+		{
+			if (Actor->IsA<AChaosCacheManager>())
+			{
+				UTakeRecorderChaosCacheSource* NewSource = TakeRecorderSources->AddSource<UTakeRecorderChaosCacheSource>();
+				NewSource->ChaosCacheManager = Cast<AChaosCacheManager>(Actor);
+
+				// Send a PropertyChangedEvent so the class catches the callback and rebuilds the property map.
+				FPropertyChangedEvent PropertyChangedEvent(UTakeRecorderChaosCacheSource::StaticClass()->FindPropertyByName(GET_MEMBER_NAME_CHECKED(UTakeRecorderChaosCacheSource, ChaosCacheManager)), EPropertyChangeType::ValueSet);
+				NewSource->PostEditChangeProperty(PropertyChangedEvent);
+			}
+		}
+	}
+}
+
+static void PopulateChaosCacheSubMenu(FMenuBuilder& MenuBuilder, UTakeRecorderSources* Sources)
+{
+	TSet<const AActor*> CacheActors;
+	
+	for (UTakeRecorderSource* Source : Sources->GetSources())
+	{
+		const UTakeRecorderChaosCacheSource* ChaosCache   = Cast<UTakeRecorderChaosCacheSource>(Source);
+		if (ChaosCache && ChaosCache->ChaosCacheManager.IsValid())
+		{
+			CacheActors.Add(ChaosCache->ChaosCacheManager.Get());
+		}
+	}
+	
+	auto OutlinerFilterPredicate = [InExistingActors = MoveTemp(CacheActors)](const AActor* InActor)
+	{
+		return !InExistingActors.Contains(InActor) && InActor->IsA<AChaosCacheManager>();
+	};
+	
+	// Set up a menu entry to add the selected actor(s) to the sequencer
+	TArray<AActor*> SelectedActors;
+	GEditor->GetSelectedActors()->GetSelectedObjects(SelectedActors);
+	SelectedActors.RemoveAll([&](const AActor* In){ return !OutlinerFilterPredicate(In); });
+	
+	FText SelectedLabel;
+	if (SelectedActors.Num() == 1)
+	{
+		SelectedLabel = FText::Format(LOCTEXT("AddSpecificActor", "Add '{0}'"), FText::FromString(SelectedActors[0]->GetActorLabel()));
+	}
+	else if (SelectedActors.Num() > 1)
+	{
+		SelectedLabel = FText::Format(LOCTEXT("AddCurrentActorSelection", "Add Current Selection ({0} actors)"), FText::AsNumber(SelectedActors.Num()));
+	}
+	
+	if (!SelectedLabel.IsEmpty())
+	{
+		MenuBuilder.AddMenuEntry(
+			SelectedLabel,
+			FText(),
+			FSlateIcon(),
+			FExecuteAction::CreateLambda([Sources, SelectedActors]{
+				AddChaosCacheSources(Sources, SelectedActors);
+			})
+		);
+	}
+	
+	MenuBuilder.BeginSection("ChooseActorSection", LOCTEXT("ChooseActor", "Choose Actor:"));
+	{
+		// Set up a menu entry to add any arbitrary actor to the sequencer
+		FSceneOutlinerInitializationOptions InitOptions;
+		{
+			// We hide the header row to keep the UI compact.
+			InitOptions.bShowHeaderRow = false;
+			InitOptions.bShowSearchBox = true;
+			InitOptions.bShowCreateNewFolder = false;
+			InitOptions.bFocusSearchBoxWhenOpened = true;
+	
+			// Only want the actor label column
+			InitOptions.ColumnMap.Add(FSceneOutlinerBuiltInColumnTypes::Label(), FSceneOutlinerColumnInfo(ESceneOutlinerColumnVisibility::Visible, 0));
+	
+			// Only display actors that are not possessed already
+			InitOptions.Filters->AddFilterPredicate<FActorTreeItem>(FActorTreeItem::FFilterPredicate::CreateLambda(OutlinerFilterPredicate));
+		}
+	
+		// actor selector to allow the user to choose an actor
+		FSceneOutlinerModule& SceneOutlinerModule = FModuleManager::LoadModuleChecked<FSceneOutlinerModule>("SceneOutliner");
+		TSharedRef< SWidget > MiniSceneOutliner =
+			SNew(SBox)
+			.MaxDesiredHeight(400.0f)
+			.WidthOverride(300.0f)
+			[
+				SceneOutlinerModule.CreateActorPicker(
+					InitOptions,
+					FOnActorPicked::CreateLambda([Sources](AActor* Actor){
+						// Create a new binding for this actor
+						FSlateApplication::Get().DismissAllMenus();
+						AddChaosCacheSources(Sources, MakeArrayView(&Actor, 1));
+					})
+				)
+			];
+	
+		MenuBuilder.AddWidget(MiniSceneOutliner, FText::GetEmpty(), true);
+	}
+	MenuBuilder.EndSection();
+}
+
+static void PopulateSourcesMenu(FMenuBuilder& MenuBuilder, UTakeRecorderSources* Sources)
+{
+	FName ExtensionName = "ChaosCacheSourceSubMenu";
+
+	MenuBuilder.AddSubMenu(
+		NSLOCTEXT("TakeRecorderSources", "ChaosCacheList_Label", "From ChaosCache"),
+		NSLOCTEXT("TakeRecorderSources", "ChaosCacheList_Tip", "Add a new recording source from a Chaos Cache Subject"),
+		FNewMenuDelegate::CreateStatic(PopulateChaosCacheSubMenu, Sources),
+		FUIAction(),
+		ExtensionName,
+		EUserInterfaceActionType::Button
+	);
+}
+
+static void ExtendSourcesMenu(TSharedRef<FExtender> Extender, UTakeRecorderSources* Sources)
+{
+	Extender->AddMenuExtension("Sources", EExtensionHook::Before, nullptr, FMenuExtensionDelegate::CreateStatic(PopulateSourcesMenu, Sources));
+}
+
+static const FName MovieSceneTrackRecorderFactoryName("MovieSceneTrackRecorderFactory");
+
 void IChaosCachingEditorPlugin::StartupModule()
 {
 	AssetTypeActions_ChaosCacheCollection = new FAssetTypeActions_ChaosCacheCollection();
@@ -44,6 +192,14 @@ void IChaosCachingEditorPlugin::StartupModule()
 	PropertyModule.RegisterCustomClassLayout("ChaosCacheCollection", FOnGetDetailCustomizationInstance::CreateStatic(&FCacheCollectionDetails::MakeInstance));
 	PropertyModule.RegisterCustomClassLayout("ChaosCacheManager", FOnGetDetailCustomizationInstance::CreateStatic(&FCacheManagerDetails::MakeInstance));
 	PropertyModule.RegisterCustomPropertyTypeLayout("ObservedComponent", FOnGetPropertyTypeCustomizationInstance::CreateStatic(&FObservedComponentDetails::MakeInstance));
+	
+	ISequencerModule& SequencerModule = FModuleManager::Get().LoadModuleChecked<ISequencerModule>("Sequencer");
+	TrackEditorBindingHandle = SequencerModule.RegisterTrackEditor(FOnCreateTrackEditor::CreateStatic(&FChaosCacheTrackEditor::CreateTrackEditor));
+	
+	ITakeRecorderModule& TakeRecorderModule = FModuleManager::Get().LoadModuleChecked<ITakeRecorderModule>("TakeRecorder");
+	SourcesMenuExtension = TakeRecorderModule.RegisterSourcesMenuExtension(FOnExtendSourcesMenu::CreateStatic(ExtendSourcesMenu));
+	
+	IModularFeatures::Get().RegisterModularFeature(MovieSceneTrackRecorderFactoryName, &MovieSceneChaosCacheTrackRecorder);
 }
 
 void IChaosCachingEditorPlugin::ShutdownModule()
@@ -70,6 +226,20 @@ void IChaosCachingEditorPlugin::ShutdownModule()
 		IAssetTools&       AssetTools       = AssetToolsModule.Get();
 
 		AssetTools.UnregisterAssetTypeActions(AssetTypeActions_ChaosCacheCollection->AsShared());
+
+		ISequencerModule* SequencerModulePtr = FModuleManager::Get().GetModulePtr<ISequencerModule>("Sequencer");
+		if (SequencerModulePtr)
+		{
+			SequencerModulePtr->UnRegisterTrackEditor(TrackEditorBindingHandle);
+		}
+
+		ITakeRecorderModule* TakeRecorderModule = FModuleManager::Get().GetModulePtr<ITakeRecorderModule>("TakeRecorder");
+		if (TakeRecorderModule)
+		{
+			TakeRecorderModule->UnregisterSourcesMenuExtension(SourcesMenuExtension);
+		}
+
+		IModularFeatures::Get().UnregisterModularFeature(MovieSceneTrackRecorderFactoryName, &MovieSceneChaosCacheTrackRecorder);
 	}
 }
 
@@ -170,7 +340,7 @@ void IChaosCachingEditorPlugin::OnCreateCacheManager()
 
 				check(Manager);
 
-				FObservedComponent* Existing = Manager->ObservedComponents.FindByPredicate([PrimitiveComp](const FObservedComponent& InItem) {
+				const FObservedComponent* Existing = Manager->GetObservedComponents().FindByPredicate([PrimitiveComp](const FObservedComponent& InItem) {
 					return InItem.GetComponent() == PrimitiveComp;
 				});
 
diff --git a/Engine/Plugins/Experimental/ChaosCaching/Source/ChaosCachingEditor/Private/Sequencer/ChaosCacheTrackEditor.cpp b/Engine/Plugins/Experimental/ChaosCaching/Source/ChaosCachingEditor/Private/Sequencer/ChaosCacheTrackEditor.cpp
new file mode 100644
--- /dev/null
+++ b/Engine/Plugins/Experimental/ChaosCaching/Source/ChaosCachingEditor/Private/Sequencer/ChaosCacheTrackEditor.cpp
@@ -0,0 +1,391 @@
+﻿// Copyright Epic Games, Inc. All Rights Reserved.
+
+#include "Sequencer/ChaosCacheTrackEditor.h"
+#include "CommonMovieSceneTools.h"
+#include "Fonts/FontMeasure.h"
+#include "Framework/MultiBox/MultiBoxBuilder.h"
+#include "GameFramework/Actor.h"
+#include "Chaos/ChaosCache.h"
+#include "Chaos/CacheManagerActor.h"
+#include "Chaos/CacheCollection.h"
+#include "LevelSequence.h"
+#include "Chaos/Sequencer/MovieSceneChaosCacheTrack.h"
+#include "Chaos/Sequencer/MovieSceneChaosCacheSection.h"
+#include "SequencerSectionPainter.h"
+#include "SequencerUtilities.h"
+#include "Styling/SlateIconFinder.h"
+#include "TimeToPixel.h"
+
+namespace ChaosCacheEditorConstants
+{
+	// @todo Sequencer Allow this to be customizable
+	constexpr float AnimationTrackHeight = 20.f;
+}
+
+#define LOCTEXT_NAMESPACE "FChaosCacheTrackEditor"
+
+static AChaosCacheManager* AcquireChaosCacheFromObjectGuid(const FGuid& Guid, TSharedPtr<ISequencer> SequencerPtr)
+{
+	UObject* BoundObject = SequencerPtr.IsValid() ? SequencerPtr->FindSpawnedObjectOrTemplate(Guid) : nullptr;
+
+	if (AActor* Actor = Cast<AActor>(BoundObject))
+	{
+		if (AChaosCacheManager* CacheManager = Cast<AChaosCacheManager>(Actor))
+		{
+			return CacheManager;
+		}
+	}
+
+	return nullptr;
+}
+
+FChaosCacheSection::FChaosCacheSection(UMovieSceneSection& InSection, TWeakPtr<ISequencer> InSequencer)
+	: Section(*CastChecked<UMovieSceneChaosCacheSection>(&InSection))
+	, Sequencer(InSequencer)
+	, InitialFirstLoopStartOffsetDuringResize(0)
+	, InitialStartTimeDuringResize(0)
+{ }
+
+UMovieSceneSection* FChaosCacheSection::GetSectionObject()
+{
+	return &Section;
+}
+
+FText FChaosCacheSection::GetSectionTitle() const
+{
+	if (Section.Params.CacheCollection != nullptr)
+	{
+		return FText::FromString(Section.Params.CacheCollection->GetName());
+	
+	}
+	return LOCTEXT("NoChaosCacheSection", "No ChaosCache");
+}
+
+float FChaosCacheSection::GetSectionHeight() const
+{
+	return ChaosCacheEditorConstants::AnimationTrackHeight;
+}
+
+int32 FChaosCacheSection::OnPaintSection(FSequencerSectionPainter& Painter) const
+{
+	using namespace UE::Sequencer;
+
+	const ESlateDrawEffect DrawEffects = Painter.bParentEnabled ? ESlateDrawEffect::None : ESlateDrawEffect::DisabledEffect;
+
+	const FTimeToPixel& TimeToPixelConverter = Painter.GetTimeConverter();
+
+	int32 LayerId = Painter.PaintSectionBackground();
+
+	static const FSlateBrush* GenericDivider = FAppStyle::GetBrush("Sequencer.GenericDivider");
+
+	if (!Section.HasStartFrame() || !Section.HasEndFrame())
+	{
+		return LayerId;
+	}
+
+	FFrameRate TickResolution = TimeToPixelConverter.GetTickResolution();
+
+	// Add lines where the animation starts and ends/loops
+	const float AnimPlayRate = FMath::IsNearlyZero(Section.Params.PlayRate) ? 1.0f : Section.Params.PlayRate;
+	const float Duration = Section.Params.GetSequenceLength();
+	const float SeqLength = Duration - TickResolution.AsSeconds(Section.Params.StartFrameOffset + Section.Params.EndFrameOffset) / AnimPlayRate;
+	const float FirstLoopSeqLength = SeqLength - TickResolution.AsSeconds(Section.Params.FirstLoopStartFrameOffset) / AnimPlayRate;
+
+	if (!FMath::IsNearlyZero(SeqLength, KINDA_SMALL_NUMBER) && SeqLength > 0)
+	{
+		float MaxOffset = Section.GetRange().Size<FFrameTime>() / TickResolution;
+		float OffsetTime = FirstLoopSeqLength;
+		float StartTime = Section.GetInclusiveStartFrame() / TickResolution;
+
+		while (OffsetTime < MaxOffset)
+		{
+			float OffsetPixel = TimeToPixelConverter.SecondsToPixel(StartTime + OffsetTime) - TimeToPixelConverter.SecondsToPixel(StartTime);
+
+			FSlateDrawElement::MakeBox(
+				Painter.DrawElements,
+				LayerId,
+				Painter.SectionGeometry.MakeChild(
+					FVector2D(2.f, Painter.SectionGeometry.Size.Y - 2.f),
+					FSlateLayoutTransform(FVector2D(OffsetPixel, 1.f))
+				).ToPaintGeometry(),
+				GenericDivider,
+				DrawEffects
+			);
+
+			OffsetTime += SeqLength;
+		}
+	}
+
+	TSharedPtr<ISequencer> SequencerPtr = Sequencer.Pin();
+	if (Painter.bIsSelected && SequencerPtr.IsValid())
+	{
+		FFrameTime CurrentTime = SequencerPtr->GetLocalTime().Time;
+		if (Section.GetRange().Contains(CurrentTime.FrameNumber) && Section.Params.CacheCollection != nullptr)
+		{
+			const float Time = TimeToPixelConverter.FrameToPixel(CurrentTime);
+
+			UChaosCacheCollection* ChaosCache = Section.Params.CacheCollection;
+
+			// Draw the current time next to the scrub handle
+			const float AnimTime = Section.MapTimeToAnimation(Duration, CurrentTime, TickResolution);
+			int32 FrameTime = FMath::FloorToFloat(TickResolution.AsFrameTime(AnimTime).AsDecimal());
+			FString FrameString = FString::FromInt(FrameTime);
+
+			const FSlateFontInfo SmallLayoutFont = FCoreStyle::GetDefaultFontStyle("Bold", 10);
+			const TSharedRef< FSlateFontMeasure > FontMeasureService = FSlateApplication::Get().GetRenderer()->GetFontMeasureService();
+			FVector2D TextSize = FontMeasureService->Measure(FrameString, SmallLayoutFont);
+
+			// Flip the text position if getting near the end of the view range
+			static const float TextOffsetPx = 10.f;
+			bool  bDrawLeft = (Painter.SectionGeometry.Size.X - Time) < (TextSize.X + 22.f) - TextOffsetPx;
+			float TextPosition = bDrawLeft ? Time - TextSize.X - TextOffsetPx : Time + TextOffsetPx;
+			//handle mirrored labels
+			const float MajorTickHeight = 9.0f;
+			FVector2D TextOffset(TextPosition, Painter.SectionGeometry.Size.Y - (MajorTickHeight + TextSize.Y));
+
+			const FLinearColor DrawColor = FAppStyle::GetSlateColor("SelectionColor").GetColor(FWidgetStyle());
+			const FVector2D BoxPadding = FVector2D(4.0f, 2.0f);
+			// draw time string
+
+			FSlateDrawElement::MakeBox(
+				Painter.DrawElements,
+				LayerId + 5,
+				Painter.SectionGeometry.ToPaintGeometry(TextOffset - BoxPadding, TextSize + 2.0f * BoxPadding),
+				FAppStyle::GetBrush("WhiteBrush"),
+				ESlateDrawEffect::None,
+				FLinearColor::Black.CopyWithNewOpacity(0.5f)
+			);
+
+			FSlateDrawElement::MakeText(
+				Painter.DrawElements,
+				LayerId + 6,
+				Painter.SectionGeometry.ToPaintGeometry(TextOffset, TextSize),
+				FrameString,
+				SmallLayoutFont,
+				DrawEffects,
+				DrawColor
+			);
+
+		}
+	}
+
+	return LayerId;
+}
+
+void FChaosCacheSection::BeginResizeSection()
+{
+	InitialFirstLoopStartOffsetDuringResize = Section.Params.FirstLoopStartFrameOffset;
+	InitialStartTimeDuringResize = Section.HasStartFrame() ? Section.GetInclusiveStartFrame() : 0;
+}
+
+void FChaosCacheSection::UpdateSection(FFrameNumber& UpdateTime) const
+{
+	const FFrameRate FrameRate = Section.GetTypedOuter<UMovieScene>()->GetTickResolution();
+	FFrameNumber StartOffset = FrameRate.AsFrameNumber((UpdateTime - InitialStartTimeDuringResize) / FrameRate * Section.Params.PlayRate);
+
+	StartOffset += InitialFirstLoopStartOffsetDuringResize;
+
+	if (StartOffset < 0)
+	{
+		// Ensure start offset is not less than 0 and adjust ResizeTime
+		UpdateTime = UpdateTime - StartOffset;
+
+		StartOffset = FFrameNumber(0);
+	}
+	else
+	{
+		// If the start offset exceeds the length of one loop, trim it back.
+		const FFrameNumber SeqLength = FrameRate.AsFrameNumber(Section.Params.GetSequenceLength()) - Section.Params.StartFrameOffset - Section.Params.EndFrameOffset;
+		StartOffset = StartOffset % SeqLength;
+	}
+
+	Section.Params.FirstLoopStartFrameOffset = StartOffset;
+}
+
+void FChaosCacheSection::ResizeSection(ESequencerSectionResizeMode ResizeMode, FFrameNumber ResizeTime)
+{
+	// Adjust the start offset when resizing from the beginning
+	if (ResizeMode == SSRM_LeadingEdge)
+	{
+		UpdateSection(ResizeTime);
+	}
+
+	ISequencerSection::ResizeSection(ResizeMode, ResizeTime);
+}
+
+void FChaosCacheSection::BeginSlipSection()
+{
+	BeginResizeSection();
+}
+
+void FChaosCacheSection::SlipSection(FFrameNumber SlipTime)
+{
+	UpdateSection(SlipTime);
+
+	ISequencerSection::SlipSection(SlipTime);
+}
+
+void FChaosCacheSection::BeginDilateSection()
+{
+	Section.PreviousPlayRate = Section.Params.PlayRate; //make sure to cache the play rate
+}
+void FChaosCacheSection::DilateSection(const TRange<FFrameNumber>& NewRange, float DilationFactor)
+{
+	Section.Params.PlayRate = Section.PreviousPlayRate / DilationFactor;
+	Section.SetRange(NewRange);
+}
+
+FChaosCacheTrackEditor::FChaosCacheTrackEditor(TSharedRef<ISequencer> InSequencer)
+	: FMovieSceneTrackEditor(InSequencer)
+{ }
+
+TSharedRef<ISequencerTrackEditor> FChaosCacheTrackEditor::CreateTrackEditor(TSharedRef<ISequencer> InSequencer)
+{
+	return MakeShareable(new FChaosCacheTrackEditor(InSequencer));
+}
+
+bool FChaosCacheTrackEditor::SupportsSequence(UMovieSceneSequence* InSequence) const
+{
+	if (InSequence && InSequence->IsTrackSupported(UMovieSceneChaosCacheTrack::StaticClass()) == ETrackSupport::NotSupported)
+	{
+		return false;
+	}
+
+	return InSequence && InSequence->IsA(ULevelSequence::StaticClass());
+}
+
+bool FChaosCacheTrackEditor::SupportsType(TSubclassOf<UMovieSceneTrack> Type) const
+{
+	return Type == UMovieSceneChaosCacheTrack::StaticClass();
+}
+
+TSharedRef<ISequencerSection> FChaosCacheTrackEditor::MakeSectionInterface(UMovieSceneSection& SectionObject, UMovieSceneTrack& Track, FGuid ObjectBinding)
+{
+	check(SupportsType(SectionObject.GetOuter()->GetClass()));
+
+	return MakeShareable(new FChaosCacheSection(SectionObject, GetSequencer()));
+}
+
+void FChaosCacheTrackEditor::BuildObjectBindingTrackMenu(FMenuBuilder& MenuBuilder, const TArray<FGuid>& ObjectBindings, const UClass* ObjectClass)
+{
+	if (ObjectClass->IsChildOf(AChaosCacheManager::StaticClass()) || ObjectClass->IsChildOf(AActor::StaticClass()))
+	{
+		if(ObjectBindings.Num() > 0)
+		{
+			AChaosCacheManager* ChaosCache = AcquireChaosCacheFromObjectGuid(ObjectBindings[0], GetSequencer());
+
+			if (ChaosCache)
+			{
+				UMovieSceneTrack* Track = nullptr;
+
+				MenuBuilder.AddMenuEntry(
+					NSLOCTEXT("Sequencer", "AddChaosCache", "Chaos Cache"),
+					NSLOCTEXT("Sequencer", "AddChaosCacheTooltip", "Adds a Chaos Cache track."),
+					FSlateIcon(),
+					FUIAction(
+						FExecuteAction::CreateSP(this, &FChaosCacheTrackEditor::BuildChaosCacheTrack, ObjectBindings, Track)
+					)
+				);
+			}
+		}
+	}
+}
+
+void FChaosCacheTrackEditor::BuildChaosCacheTrack(TArray<FGuid> ObjectBindings, UMovieSceneTrack* Track)
+{
+	const TSharedPtr<ISequencer> SequencerPtr = GetSequencer();
+
+	if (SequencerPtr.IsValid())
+	{
+		const FScopedTransaction Transaction(LOCTEXT("AddChaosCache_Transaction", "Add Chaos Cache"));
+
+		for (FGuid ObjectBinding : ObjectBindings)
+		{
+			if (ObjectBinding.IsValid())
+			{
+				UObject* Object = SequencerPtr->FindSpawnedObjectOrTemplate(ObjectBinding);
+
+				AChaosCacheManager* ChaosCache = AcquireChaosCacheFromObjectGuid(ObjectBinding, GetSequencer());
+
+				if (Object && ChaosCache)
+				{
+					AnimatablePropertyChanged(FOnKeyProperty::CreateRaw(this, &FChaosCacheTrackEditor::AddKeyInternal, Object, ChaosCache, Track));
+				}
+			}
+		}
+	}
+}
+
+FKeyPropertyResult FChaosCacheTrackEditor::AddKeyInternal(FFrameNumber KeyTime, UObject* Object, AChaosCacheManager* ChaosCache, UMovieSceneTrack* Track)
+{
+	FKeyPropertyResult KeyPropertyResult;
+
+	const FFindOrCreateHandleResult HandleResult = FindOrCreateHandleToObject(Object);
+	const FGuid ObjectHandle = HandleResult.Handle;
+	KeyPropertyResult.bHandleCreated |= HandleResult.bWasCreated;
+	if (ObjectHandle.IsValid())
+	{
+		if (!Track)
+		{
+			Track = AddTrack(GetSequencer()->GetFocusedMovieSceneSequence()->GetMovieScene(), ObjectHandle, UMovieSceneChaosCacheTrack::StaticClass(), NAME_None);
+			KeyPropertyResult.bTrackCreated = true;
+		}
+
+		if (ensure(Track))
+		{
+			Track->Modify();
+
+			UMovieSceneSection* NewSection = Cast<UMovieSceneChaosCacheTrack>(Track)->AddNewAnimation(KeyTime, ChaosCache);
+			KeyPropertyResult.bTrackModified = true;
+			KeyPropertyResult.SectionsCreated.Add(NewSection);
+
+			GetSequencer()->EmptySelection();
+			GetSequencer()->SelectSection(NewSection);
+			GetSequencer()->ThrobSectionSelection();
+		}
+	}
+
+	return KeyPropertyResult;
+}
+
+TSharedPtr<SWidget> FChaosCacheTrackEditor::BuildOutlinerEditWidget(const FGuid& ObjectBinding, UMovieSceneTrack* Track, const FBuildEditWidgetParams& Params)
+{
+	AChaosCacheManager* ChaosCache = AcquireChaosCacheFromObjectGuid(ObjectBinding, GetSequencer());
+
+	if (ChaosCache)
+	{
+		TWeakPtr<ISequencer> WeakSequencer = GetSequencer();
+
+		auto SubMenuCallback = [=]() -> TSharedRef<SWidget>
+		{
+			FMenuBuilder MenuBuilder(true, nullptr);
+
+			TArray<FGuid> ObjectBindings;
+			ObjectBindings.Add(ObjectBinding);
+
+			BuildChaosCacheTrack(ObjectBindings, Track);
+			
+			return MenuBuilder.MakeWidget();
+		};
+
+		return SNew(SHorizontalBox)
+			+ SHorizontalBox::Slot()
+			.AutoWidth()
+			.VAlign(VAlign_Center)
+			[
+				FSequencerUtilities::MakeAddButton(LOCTEXT("ChaosCacheText", "Chaos Cache"), FOnGetContent::CreateLambda(SubMenuCallback), Params.NodeIsHovered, GetSequencer())
+			];
+	}
+	else
+	{
+		return TSharedPtr<SWidget>();
+	}
+
+}
+
+const FSlateBrush* FChaosCacheTrackEditor::GetIconBrush() const
+{
+	return FSlateIconFinder::FindIconForClass(AChaosCacheManager::StaticClass()).GetIcon();
+}
+
+#undef LOCTEXT_NAMESPACE
diff --git a/Engine/Plugins/Experimental/ChaosCaching/Source/ChaosCachingEditor/Private/Sequencer/MovieSceneChaosCacheTrackRecorder.cpp b/Engine/Plugins/Experimental/ChaosCaching/Source/ChaosCachingEditor/Private/Sequencer/MovieSceneChaosCacheTrackRecorder.cpp
new file mode 100644
--- /dev/null
+++ b/Engine/Plugins/Experimental/ChaosCaching/Source/ChaosCachingEditor/Private/Sequencer/MovieSceneChaosCacheTrackRecorder.cpp
@@ -0,0 +1,119 @@
+﻿// Copyright Epic Games, Inc. All Rights Reserved.
+
+#include "Sequencer/MovieSceneChaosCacheTrackRecorder.h"
+#include "Chaos/CacheManagerActor.h"
+#include "Chaos/Sequencer/MovieSceneChaosCacheSection.h"
+#include "Chaos/Sequencer/MovieSceneChaosCacheTrack.h"
+#include "TakeRecorderSettings.h"
+#include "Recorder/TakeRecorderParameters.h"
+#include "Engine/TimecodeProvider.h"
+#include "Misc/App.h"
+#include "MovieSceneFolder.h"
+#include "Tracks/MovieSceneActorReferenceTrack.h"
+#include "LevelSequence.h"
+
+#define LOCTEXT_NAMESPACE "MovieSceneChaosCacheTrackRecorder"
+
+bool FMovieSceneChaosCacheTrackRecorderFactory::CanRecordObject(UObject* InObjectToRecord) const
+{
+	return InObjectToRecord->IsA<AChaosCacheManager>();
+}
+
+UMovieSceneTrackRecorder* FMovieSceneChaosCacheTrackRecorderFactory::CreateTrackRecorderForObject() const
+{
+	return NewObject<UMovieSceneChaosCacheTrackRecorder>();
+}
+
+void UMovieSceneChaosCacheTrackRecorder::CreateTrackImpl()
+{
+	ChaosCacheTrack = MovieScene->AddTrack<UMovieSceneChaosCacheTrack>(ObjectGuid);\
+	ChaosCacheManager = Cast<AChaosCacheManager>(ObjectToRecord.Get());
+	
+	ChaosCacheSection = Cast<UMovieSceneChaosCacheSection>(ChaosCacheTrack->CreateNewSection());
+	if (ChaosCacheSection.IsValid())
+	{
+		ChaosCacheSection->SetIsActive(false);
+		ChaosCacheTrack->AddSection(*ChaosCacheSection);
+		ChaosCacheSection->Params.CacheCollection = ChaosCacheManager->CacheCollection;
+
+		// Resize the section to either it's remaining	keyframes range or 0
+		ChaosCacheSection->SetRange(ChaosCacheSection->GetAutoSizeRange().Get(TRange<FFrameNumber>(0, 0)));
+
+		// Make sure it starts at frame 0, in case Auto Size removed a piece of the start
+		ChaosCacheSection->ExpandToFrame(0);
+	}
+}
+
+void UMovieSceneChaosCacheTrackRecorder::SetSectionStartTimecodeImpl(const FTimecode& InSectionStartTimecode, const FFrameNumber& InSectionFirstFrame)
+{
+	if (ChaosCacheSection.IsValid() && ChaosCacheTrack.IsValid() && ChaosCacheManager)
+	{
+		ChaosCacheSection->TimecodeSource = FMovieSceneTimecodeSource(InSectionStartTimecode);
+
+		FTakeRecorderParameters Parameters;
+		Parameters.User = GetDefault<UTakeRecorderUserSettings>()->Settings;
+		Parameters.Project = GetDefault<UTakeRecorderProjectSettings>()->Settings;
+
+		const FFrameRate TickResolution = MovieScene->GetTickResolution();
+		const FFrameRate DisplayRate = MovieScene->GetDisplayRate();
+
+		RecordStartTime = FApp::GetCurrentTime();
+		RecordStartFrame = Parameters.Project.bStartAtCurrentTimecode ?
+			FFrameRate::TransformTime(FFrameTime(InSectionStartTimecode.ToFrameNumber(DisplayRate)),
+				DisplayRate, TickResolution).FloorToFrame() : MovieScene->GetPlaybackRange().GetLowerBoundValue();
+
+		StoredCacheMode = ChaosCacheManager->CacheMode;
+		ChaosCacheManager->CacheMode = ECacheMode::Record;
+		ChaosCacheManager->BeginEvaluate();
+	}
+}
+
+UMovieSceneSection* UMovieSceneChaosCacheTrackRecorder::GetMovieSceneSection() const
+{
+	return Cast<UMovieSceneSection>(ChaosCacheSection.Get());
+}
+
+void UMovieSceneChaosCacheTrackRecorder::StopRecordingImpl()
+{
+	if(ChaosCacheManager)
+	{
+		// Stop the async recorder
+		ChaosCacheManager->EndEvaluate();
+		ChaosCacheManager->CacheMode = StoredCacheMode;
+	}
+}
+
+void UMovieSceneChaosCacheTrackRecorder::FinalizeTrackImpl()
+{
+	if (ChaosCacheSection.IsValid())
+	{
+		// Set the final range 
+		TOptional<TRange<FFrameNumber> > DefaultSectionLength = ChaosCacheSection->GetAutoSizeRange();
+		if (DefaultSectionLength.IsSet())
+		{
+			ChaosCacheSection->SetRange(DefaultSectionLength.GetValue());
+		}
+
+		// Activate the section
+		ChaosCacheSection->SetIsActive(true);
+	}
+}
+
+void UMovieSceneChaosCacheTrackRecorder::RecordSampleImpl(const FQualifiedFrameTime& CurrentFrameTime)
+{
+	// Expand the section to the new length
+	const FFrameRate	 TickResolution = MovieScene->GetTickResolution();
+	const FFrameNumber CurrentFrame = CurrentFrameTime.ConvertTo(TickResolution).FloorToFrame();
+
+	if (ChaosCacheSection.IsValid())
+	{
+		ChaosCacheSection->SetEndFrame(CurrentFrame);
+	}
+}
+
+bool UMovieSceneChaosCacheTrackRecorder::LoadRecordedFile(const FString& FileName, UMovieScene *InMovieScene, TMap<FGuid, AActor*>& ActorGuidToActorMap, TFunction<void()> InCompletionCallback)
+{
+	return false;
+}
+
+#undef LOCTEXT_NAMESPACE
diff --git a/Engine/Plugins/Experimental/ChaosCaching/Source/ChaosCachingEditor/Private/Sequencer/TakeRecorderChaosCacheSource.cpp b/Engine/Plugins/Experimental/ChaosCaching/Source/ChaosCachingEditor/Private/Sequencer/TakeRecorderChaosCacheSource.cpp
new file mode 100644
--- /dev/null
+++ b/Engine/Plugins/Experimental/ChaosCaching/Source/ChaosCachingEditor/Private/Sequencer/TakeRecorderChaosCacheSource.cpp
@@ -0,0 +1,112 @@
+﻿// Copyright Epic Games, Inc. All Rights Reserved.
+
+#include "Sequencer/TakeRecorderChaosCacheSource.h"
+
+#include "Chaos/CacheManagerActor.h"
+#include "Chaos/Sequencer/MovieSceneChaosCacheSection.h"
+#include "Chaos/Sequencer/MovieSceneChaosCacheTrack.h"
+#include "Sequencer/MovieSceneChaosCacheTrackRecorder.h"
+
+#include "LevelSequence.h"
+#include "MovieSceneFolder.h"
+#include "Channels/MovieSceneFloatChannel.h"
+#include "Channels/MovieSceneChannelTraits.h"
+#include "Framework/Notifications/NotificationManager.h" 
+#include "Widgets/Notifications/SNotificationList.h"
+
+
+#define LOCTEXT_NAMESPACE "TakeRecorderChaosCacheSource"
+
+UTakeRecorderChaosCacheSource::UTakeRecorderChaosCacheSource(const FObjectInitializer& ObjInit)
+	: Super(ObjInit)
+	, ChaosCacheManager(nullptr)
+{
+	TrackTint = FColor(0.0f, 125.0f, 255.0f, 65.0f);
+}
+
+TArray<UTakeRecorderSource*> UTakeRecorderChaosCacheSource::PreRecording(ULevelSequence* InSequence, FMovieSceneSequenceID InSequenceID, ULevelSequence* InMasterSequence, FManifestSerializer* InManifestSerializer)
+{
+	if (ChaosCacheManager)
+	{
+		UMovieScene* MovieScene = InMasterSequence->GetMovieScene();
+		TrackRecorder = NewObject<UMovieSceneChaosCacheTrackRecorder>();
+
+		// We only support possessable for now since the spawnable template is clearing out the component reference
+		// Would be better to use a TSoftObjectPtr for the actor and have a mechanism similar to the FixUpPIE to go
+		// from PIE to editor to have a correct soft object path
+		const FGuid ObjectGuid = MovieScene->AddPossessable(ChaosCacheManager->GetActorLabel(), ChaosCacheManager->GetClass());
+		InMasterSequence->BindPossessableObject(ObjectGuid, *ChaosCacheManager, ChaosCacheManager->GetWorld());
+		
+		TrackRecorder->CreateTrack(nullptr, ChaosCacheManager.Get(), MovieScene, nullptr, ObjectGuid);
+		CachedChaosCacheTrack = TrackRecorder->GetChaosCacheTrack();
+	}
+
+	return TArray<UTakeRecorderSource*>();
+}
+
+void UTakeRecorderChaosCacheSource::TickRecording(const FQualifiedFrameTime& CurrentTime)
+{
+	if (TrackRecorder)
+	{
+		TrackRecorder->RecordSample(CurrentTime);
+	}
+}
+
+void UTakeRecorderChaosCacheSource::StartRecording(const FTimecode& InSectionStartTimecode, const FFrameNumber& InSectionFirstFrame, class ULevelSequence* InSequence)
+{
+	if (TrackRecorder)
+	{
+		TrackRecorder->SetSectionStartTimecode(InSectionStartTimecode, InSectionFirstFrame);
+	}
+}
+
+void UTakeRecorderChaosCacheSource::StopRecording(class ULevelSequence* InSequence)
+{
+	if (TrackRecorder)
+	{
+		TrackRecorder->StopRecording();
+	}
+}
+
+TArray<UTakeRecorderSource*> UTakeRecorderChaosCacheSource::PostRecording(class ULevelSequence* InSequence, ULevelSequence* InMasterSequence, const bool bCancelled)
+{
+	if (TrackRecorder)
+	{
+		TrackRecorder->FinalizeTrack();
+	}
+
+	TrackRecorder = nullptr;
+	return TArray<UTakeRecorderSource*>();
+}
+
+void UTakeRecorderChaosCacheSource::AddContentsToFolder(UMovieSceneFolder* InFolder)
+{
+	if (CachedChaosCacheTrack.IsValid())
+	{
+		InFolder->AddChildMasterTrack(CachedChaosCacheTrack.Get());
+	}
+}
+
+FText UTakeRecorderChaosCacheSource::GetDisplayTextImpl() const
+{
+	if (ChaosCacheManager != nullptr && ChaosCacheManager->IsValidLowLevelFast())
+	{
+		return FText::FromString(ChaosCacheManager->GetName());
+	}
+
+	return LOCTEXT("Display Text", "Chaos Cache");
+}
+
+void UTakeRecorderChaosCacheSource::PostEditChangeProperty(FPropertyChangedEvent& PropertyChangedEvent)
+{
+	Super::PostEditChangeProperty(PropertyChangedEvent);
+
+	const FName PropertyName = PropertyChangedEvent.GetPropertyName();
+}
+
+void UTakeRecorderChaosCacheSource::PostLoad()
+{
+	Super::PostLoad();
+}
+
+#undef LOCTEXT_NAMESPACE
diff --git a/Engine/Plugins/Experimental/ChaosCaching/Source/ChaosCachingEditor/Public/Chaos/ChaosCachingEditorPlugin.h b/Engine/Plugins/Experimental/ChaosCaching/Source/ChaosCachingEditor/Public/Chaos/ChaosCachingEditorPlugin.h
--- a/Engine/Plugins/Experimental/ChaosCaching/Source/ChaosCachingEditor/Public/Chaos/ChaosCachingEditorPlugin.h
+++ b/Engine/Plugins/Experimental/ChaosCaching/Source/ChaosCachingEditor/Public/Chaos/ChaosCachingEditorPlugin.h
@@ -7,6 +7,7 @@
 #include "Modules/ModuleManager.h"
 #include "Delegates/IDelegateInstance.h"
 #include "Framework/MultiBox/MultiBoxExtender.h"
+#include "Sequencer/MovieSceneChaosCacheTrackRecorder.h"
 
 class FAssetTypeActions_ChaosCacheCollection;
 
@@ -52,4 +53,8 @@ private:
 
 	FAssetTypeActions_ChaosCacheCollection* AssetTypeActions_ChaosCacheCollection;
 	FDelegateHandle StartupHandle;
+	FDelegateHandle TrackEditorBindingHandle;
+	FDelegateHandle SourcesMenuExtension;
+
+	FMovieSceneChaosCacheTrackRecorderFactory MovieSceneChaosCacheTrackRecorder;
 };
diff --git a/Engine/Plugins/Experimental/ChaosCaching/Source/ChaosCachingEditor/Public/Sequencer/ChaosCacheTrackEditor.h b/Engine/Plugins/Experimental/ChaosCaching/Source/ChaosCachingEditor/Public/Sequencer/ChaosCacheTrackEditor.h
new file mode 100644
--- /dev/null
+++ b/Engine/Plugins/Experimental/ChaosCaching/Source/ChaosCachingEditor/Public/Sequencer/ChaosCacheTrackEditor.h
@@ -0,0 +1,97 @@
+﻿// Copyright Epic Games, Inc. All Rights Reserved.
+
+#pragma once
+
+#include "ISequencer.h"
+#include "ISequencerSection.h"
+#include "ISequencerTrackEditor.h"
+#include "MovieSceneTrackEditor.h"
+
+/**
+ * Editor for chaos Cache tracks
+ */
+class FChaosCacheTrackEditor : public FMovieSceneTrackEditor
+{
+public:
+	/**
+	 * Constructor
+	 *
+	 * @param InSequencer The sequencer instance to be used by this tool
+	 */
+	FChaosCacheTrackEditor( TSharedRef<ISequencer> InSequencer );
+
+	/** Virtual destructor. */
+	virtual ~FChaosCacheTrackEditor() { }
+
+	/**
+	 * Creates an instance of this class.  Called by a sequencer 
+	 *
+	 * @param OwningSequencer The sequencer instance to be used by this tool
+	 * @return The new instance of this class
+	 */
+	static TSharedRef<ISequencerTrackEditor> CreateTrackEditor( TSharedRef<ISequencer> OwningSequencer );
+
+public:
+
+	// ~Begin ISequencerTrackEditor interface
+	virtual void BuildObjectBindingTrackMenu(class FMenuBuilder& MenuBuilder, const TArray<FGuid>& ObjectBindings, const UClass* ObjectClass) override;
+	virtual TSharedRef<ISequencerSection> MakeSectionInterface(UMovieSceneSection& SectionObject, UMovieSceneTrack& Track, FGuid ObjectBinding) override;
+	virtual bool SupportsSequence(UMovieSceneSequence* InSequence) const override;
+	virtual bool SupportsType(TSubclassOf<UMovieSceneTrack> Type) const override;
+	virtual TSharedPtr<SWidget> BuildOutlinerEditWidget(const FGuid& ObjectBinding, UMovieSceneTrack* Track, const FBuildEditWidgetParams& Params) override;
+	virtual const FSlateBrush* GetIconBrush() const override;
+	// ~End ISequencerTrackEditor interface
+	
+private:
+
+	/** Build an internal chaos cache track */
+	void BuildChaosCacheTrack(TArray<FGuid> ObjectBindings, UMovieSceneTrack* Track);
+
+	/** Delegate for AnimatablePropertyChanged in AddKey */
+	FKeyPropertyResult AddKeyInternal(FFrameNumber KeyTime, UObject* Object, class AChaosCacheManager* ChaosCache, UMovieSceneTrack* Track);
+};
+
+/** Chaos cache sequencer section */
+class FChaosCacheSection
+	: public ISequencerSection
+	, public TSharedFromThis<FChaosCacheSection>
+{
+public:
+
+	/** Constructor. */
+	FChaosCacheSection( UMovieSceneSection& InSection, TWeakPtr<ISequencer> InSequencer);
+
+	/** Virtual destructor. */
+	virtual ~FChaosCacheSection() { }
+
+public:
+
+	// ~Begin ISequencerSection interface
+	virtual UMovieSceneSection* GetSectionObject() override;
+	virtual FText GetSectionTitle() const override;
+	virtual float GetSectionHeight() const override;
+	virtual int32 OnPaintSection( FSequencerSectionPainter& Painter ) const override;
+	virtual void BeginResizeSection() override;
+	virtual void ResizeSection(ESequencerSectionResizeMode ResizeMode, FFrameNumber ResizeTime) override;
+	virtual void BeginSlipSection() override;
+	virtual void SlipSection(FFrameNumber SlipTime) override;
+	virtual void BeginDilateSection() override;
+	virtual void DilateSection(const TRange<FFrameNumber>& NewRange, float DilationFactor) override;
+	// ~End ISequencerSection interface
+	
+private:
+
+	void UpdateSection(FFrameNumber& UpdateTime) const;
+
+	/** The section we are visualizing */
+	class UMovieSceneChaosCacheSection& Section;
+
+	/** Used to draw animation frame, need selection state and local time*/
+	TWeakPtr<ISequencer> Sequencer;
+
+	/** Cached first loop start offset value valid only during resize */
+	FFrameNumber InitialFirstLoopStartOffsetDuringResize;
+	
+	/** Cached start time valid only during resize */
+	FFrameNumber InitialStartTimeDuringResize;
+};
diff --git a/Engine/Plugins/Experimental/ChaosCaching/Source/ChaosCachingEditor/Public/Sequencer/MovieSceneChaosCacheTrackRecorder.h b/Engine/Plugins/Experimental/ChaosCaching/Source/ChaosCachingEditor/Public/Sequencer/MovieSceneChaosCacheTrackRecorder.h
new file mode 100644
--- /dev/null
+++ b/Engine/Plugins/Experimental/ChaosCaching/Source/ChaosCachingEditor/Public/Sequencer/MovieSceneChaosCacheTrackRecorder.h
@@ -0,0 +1,78 @@
+﻿// Copyright Epic Games, Inc. All Rights Reserved.
+
+#pragma once
+
+#include "TrackRecorders/MovieSceneTrackRecorder.h"
+#include "CoreMinimal.h"
+#include "MovieScene.h"
+#include "TrackRecorders/IMovieSceneTrackRecorderFactory.h"
+#include "Channels/MovieSceneFloatChannel.h"
+#include "Chaos/CacheManagerActor.h"
+
+#include "MovieSceneChaosCacheTrackRecorder.generated.h"
+
+class  UMovieSceneChaosCacheTrack;
+class  UMovieSceneChaosCacheSection;
+
+struct FFrameNumber;
+class  UMovieSceneSection;
+class  UMovieSceneTrackRecorderSettings;
+
+class CHAOSCACHINGEDITOR_API FMovieSceneChaosCacheTrackRecorderFactory : public IMovieSceneTrackRecorderFactory
+{
+public:
+	virtual ~FMovieSceneChaosCacheTrackRecorderFactory() {}
+
+	// ~Begin IMovieSceneTrackRecorderFactory Interface
+	virtual bool CanRecordObject(class UObject* InObjectToRecord) const override;
+	virtual UMovieSceneTrackRecorder* CreateTrackRecorderForObject() const override;
+	virtual bool CanRecordProperty(class UObject* InObjectToRecord, class FProperty* InPropertyToRecord) const override { return false; }
+	virtual UMovieSceneTrackRecorder* CreateTrackRecorderForProperty(UObject* InObjectToRecord, const FName& InPropertyToRecord) const override { return nullptr; }
+	virtual FText GetDisplayName() const override { return NSLOCTEXT("MovieSceneChaosCacheTrackRecorderFactory", "DisplayName", "Chaos Cache Track"); }
+	// ~End IMovieSceneTrackRecorderFactory Interface
+};
+
+/**
+* Track recorder implementation for the chaos cache
+*/
+UCLASS(BlueprintType)
+class CHAOSCACHINGEDITOR_API UMovieSceneChaosCacheTrackRecorder
+	: public UMovieSceneTrackRecorder
+{
+	GENERATED_BODY()
+
+public:
+	// ~Begin UMovieSceneTrackRecorder Interface
+	virtual void RecordSampleImpl(const FQualifiedFrameTime& CurrentFrameTime) override;
+	virtual void FinalizeTrackImpl() override;
+	virtual void SetSectionStartTimecodeImpl(const FTimecode& InSectionStartTimecode, const FFrameNumber& InSectionFirstFrame) override;
+	virtual UMovieSceneSection* GetMovieSceneSection() const override;
+	virtual void StopRecordingImpl() override;
+	virtual bool LoadRecordedFile(const FString& InFileName, UMovieScene *InMovieScene, TMap<FGuid, AActor*>& ActorGuidToActorMap, TFunction<void()> InCompletionCallback) override;
+	virtual void CreateTrackImpl() override;
+	// ~End UMovieSceneTrackRecorder Interface
+
+	/** Returns the chaos cache track on which the cache manager will be recorded */
+	TWeakObjectPtr<UMovieSceneChaosCacheTrack> GetChaosCacheTrack() const {return ChaosCacheTrack;}
+	
+private:
+	
+	/** The ChaosCache Track to record onto */
+	TWeakObjectPtr<UMovieSceneChaosCacheTrack> ChaosCacheTrack;
+
+	/** Sections to record to on each track*/
+	TWeakObjectPtr<UMovieSceneChaosCacheSection> ChaosCacheSection;
+
+	/** Chaos cache that will be used to record the simulation */
+	AChaosCacheManager* ChaosCacheManager;
+
+	/** Stored cahce mode that will be set back onto the manager once the recording will be finished */
+	ECacheMode StoredCacheMode;
+
+	/** The time at the start of this recording section */
+	double RecordStartTime;
+
+	/** The frame at the start of this recording section */
+	FFrameNumber RecordStartFrame;
+};
+
diff --git a/Engine/Plugins/Experimental/ChaosCaching/Source/ChaosCachingEditor/Public/Sequencer/TakeRecorderChaosCacheSource.h b/Engine/Plugins/Experimental/ChaosCaching/Source/ChaosCachingEditor/Public/Sequencer/TakeRecorderChaosCacheSource.h
new file mode 100644
--- /dev/null
+++ b/Engine/Plugins/Experimental/ChaosCaching/Source/ChaosCachingEditor/Public/Sequencer/TakeRecorderChaosCacheSource.h
@@ -0,0 +1,62 @@
+﻿// Copyright Epic Games, Inc. All Rights Reserved.
+
+#pragma once
+
+#include "CoreMinimal.h"
+#include "TakeRecorderSource.h"
+#include "UObject/WeakObjectPtrTemplates.h"
+
+#include "TakeRecorderChaosCacheSource.generated.h"
+
+class ULevelSequence;
+class UMovieSceneFolder;
+class UMovieSceneChaosCacheTrack;
+class UMovieSceneChaosCacheTrackRecorder;
+class AChaosCacheManager;
+
+/** A recording source selector for the chaos integration into take recorder */
+
+UCLASS(Category = "Chaos", meta = (TakeRecorderDisplayName = "Chaos Recorder"))
+class UTakeRecorderChaosCacheSource
+	: public UTakeRecorderSource
+{
+public:
+	GENERATED_BODY()
+
+	UTakeRecorderChaosCacheSource(const FObjectInitializer& ObjInit);
+
+	/** Chaos Cache manager to be used as take recorder source */
+	UPROPERTY(EditAnywhere, BlueprintReadWrite, Category = "Source", meta = (DisplayName = "Chaos Cache"))
+	TSoftObjectPtr<AChaosCacheManager> ChaosCacheManager;
+	
+	// ~Begin UObject Interface
+	virtual void PostEditChangeProperty(FPropertyChangedEvent& PropertyChangedEvent) override;
+	virtual void PostLoad() override;
+	// ~End UObject Interface
+	
+protected:
+	// ~Begin UTakeRecorderSource Interface
+	virtual bool SupportsSubscenes() const override { return false; }
+	// ~End UTakeRecorderSource Interface
+
+private:
+	// ~Begin UTakeRecorderSource Interface
+	virtual TArray<UTakeRecorderSource*> PreRecording(ULevelSequence* InSequence, FMovieSceneSequenceID InSequenceID, ULevelSequence* InMasterSequence, FManifestSerializer* InManifestSerializer) override;
+	virtual void StartRecording(const FTimecode& InSectionStartTimecode, const FFrameNumber& InSectionFirstFrame, class ULevelSequence* InSequence) override;
+	virtual void StopRecording(class ULevelSequence* InSequence) override;
+	virtual void TickRecording(const FQualifiedFrameTime& CurrentTime) override;
+	virtual TArray<UTakeRecorderSource*> PostRecording(class ULevelSequence* InSequence, ULevelSequence* InMasterSequence, const bool bCancelled) override;
+	virtual void AddContentsToFolder(UMovieSceneFolder* InFolder) override;
+	virtual FText GetDisplayTextImpl() const override;
+	// ~End UTakeRecorderSource
+	
+	/** Chaos cache track recorder used by this source */
+	UPROPERTY()
+	TObjectPtr<UMovieSceneChaosCacheTrackRecorder> TrackRecorder;
+
+	/**
+	 * Stores an existing Chaos cache track track in the Sequence to be recorded or
+	 * a new one created for recording. Set during PreRecording.
+	 */
+	TWeakObjectPtr<UMovieSceneChaosCacheTrack> CachedChaosCacheTrack;
+};
diff --git a/Engine/Plugins/Runtime/GeometryCache/Source/GeometryCacheTracks/Private/MovieSceneGeometryCacheSection.cpp b/Engine/Plugins/Runtime/GeometryCache/Source/GeometryCacheTracks/Private/MovieSceneGeometryCacheSection.cpp
--- a/Engine/Plugins/Runtime/GeometryCache/Source/GeometryCacheTracks/Private/MovieSceneGeometryCacheSection.cpp
+++ b/Engine/Plugins/Runtime/GeometryCache/Source/GeometryCacheTracks/Private/MovieSceneGeometryCacheSection.cpp
@@ -108,22 +108,6 @@ void UMovieSceneGeometryCacheSection::Serialize(FArchive& Ar)
 	Super::Serialize(Ar);
 }
 
-FFrameNumber GetFirstLoopStartOffsetAtTrimTime(FQualifiedFrameTime TrimTime, const FMovieSceneGeometryCacheParams& Params, FFrameNumber StartFrame, FFrameRate FrameRate)
-{
-	const float AnimPlayRate = FMath::IsNearlyZero(Params.PlayRate) ? 1.0f : Params.PlayRate;
-	const float AnimPosition = (TrimTime.Time - StartFrame) / TrimTime.Rate * AnimPlayRate;
-	const float SeqLength = Params.GetSequenceLength() - FrameRate.AsSeconds(Params.StartFrameOffset + Params.EndFrameOffset) / AnimPlayRate;
-
-	FFrameNumber NewOffset = FrameRate.AsFrameNumber(FMath::Fmod(AnimPosition, SeqLength));
-	NewOffset += Params.FirstLoopStartFrameOffset;
-
-	const FFrameNumber SeqLengthInFrames = FrameRate.AsFrameNumber(SeqLength);
-	NewOffset = NewOffset % SeqLengthInFrames;
-
-	return NewOffset;
-}
-
-
 TOptional<TRange<FFrameNumber> > UMovieSceneGeometryCacheSection::GetAutoSizeRange() const
 {
 	FFrameRate FrameRate = GetTypedOuter<UMovieScene>()->GetTickResolution();
diff --git a/Engine/Plugins/Runtime/HairStrands/Source/HairStrandsCore/Private/MovieSceneGroomCacheSection.cpp b/Engine/Plugins/Runtime/HairStrands/Source/HairStrandsCore/Private/MovieSceneGroomCacheSection.cpp
--- a/Engine/Plugins/Runtime/HairStrands/Source/HairStrandsCore/Private/MovieSceneGroomCacheSection.cpp
+++ b/Engine/Plugins/Runtime/HairStrands/Source/HairStrandsCore/Private/MovieSceneGroomCacheSection.cpp
@@ -32,22 +32,6 @@ TOptional<FFrameTime> UMovieSceneGroomCacheSection::GetOffsetTime() const
 	return TOptional<FFrameTime>(Params.FirstLoopStartFrameOffset);
 }
 
-FFrameNumber GetFirstLoopStartOffsetAtTrimTime(FQualifiedFrameTime TrimTime, const FMovieSceneGroomCacheParams& Params, FFrameNumber StartFrame, FFrameRate FrameRate)
-{
-	const float AnimPlayRate = FMath::IsNearlyZero(Params.PlayRate) ? 1.0f : Params.PlayRate;
-	const float AnimPosition = (TrimTime.Time - StartFrame) / TrimTime.Rate * AnimPlayRate;
-	const float SeqLength = Params.GetSequenceLength() - FrameRate.AsSeconds(Params.StartFrameOffset + Params.EndFrameOffset) / AnimPlayRate;
-
-	FFrameNumber NewOffset = FrameRate.AsFrameNumber(FMath::Fmod(AnimPosition, SeqLength));
-	NewOffset += Params.FirstLoopStartFrameOffset;
-
-	const FFrameNumber SeqLengthInFrames = FrameRate.AsFrameNumber(SeqLength);
-	NewOffset = NewOffset % SeqLengthInFrames;
-
-	return NewOffset;
-}
-
-
 TOptional<TRange<FFrameNumber> > UMovieSceneGroomCacheSection::GetAutoSizeRange() const
 {
 	FFrameRate FrameRate = GetTypedOuter<UMovieScene>()->GetTickResolution();
diff --git a/Engine/Source/Runtime/LevelSequence/Private/LevelSequenceModule.cpp b/Engine/Source/Runtime/LevelSequence/Private/LevelSequenceModule.cpp
--- a/Engine/Source/Runtime/LevelSequence/Private/LevelSequenceModule.cpp
+++ b/Engine/Source/Runtime/LevelSequence/Private/LevelSequenceModule.cpp
@@ -320,7 +320,7 @@ FDelegateHandle FLevelSequenceModule::RegisterObjectSpawner(FOnCreateMovieSceneO
 	return OnCreateMovieSceneObjectSpawnerDelegates.Last().GetHandle();
 }
 
-void FLevelSequenceModule::UnregisterObjectSpawner(FDelegateHandle InHandle)
+void FLevelSequenceModule::UnregisterObjectSpawner(FDelegateHandle InHandle) 
 {
 	OnCreateMovieSceneObjectSpawnerDelegates.RemoveAll([=](const FOnCreateMovieSceneObjectSpawner& Delegate) { return Delegate.GetHandle() == InHandle; });
 }
@@ -337,7 +337,7 @@ void FLevelSequenceModule::GenerateObjectSpawners(TArray<TSharedRef<IMovieSceneO
 	// @TODO: we could also sort by most-derived type here to allow for type specific behaviors
 	OutSpawners.Sort([](TSharedRef<IMovieSceneObjectSpawner> LHS, TSharedRef<IMovieSceneObjectSpawner> RHS)
 	{
-		return LHS->IsEditor() > RHS->IsEditor();
+		return (LHS->IsEditor() > RHS->IsEditor()) || ( (LHS->IsEditor() == RHS->IsEditor()) && (LHS->GetSpawnerPriority() > RHS->GetSpawnerPriority()));
 	});
 }
 
diff --git a/Engine/Source/Runtime/MovieScene/Public/MovieSceneSection.h b/Engine/Source/Runtime/MovieScene/Public/MovieSceneSection.h
--- a/Engine/Source/Runtime/MovieScene/Public/MovieSceneSection.h
+++ b/Engine/Source/Runtime/MovieScene/Public/MovieSceneSection.h
@@ -18,6 +18,7 @@
 #include "Math/RangeBound.h"
 #include "Misc/AssertionMacros.h"
 #include "Misc/FrameNumber.h"
+#include "Misc/QualifiedFrameTime.h"
 #include "Misc/FrameRate.h"
 #include "Misc/FrameTime.h"
 #include "Misc/Optional.h"
@@ -38,7 +39,6 @@
 #include "Generators/MovieSceneEasingFunction.h"
 #include "KeyParams.h"
 #include "Misc/FrameTime.h"
-#include "Misc/QualifiedFrameTime.h"
 #include "MovieScene.h"
 #include "MovieSceneFrameMigration.h"
 #include "MovieSceneFwd.h"
@@ -788,3 +788,19 @@ protected:
 	/** Defines whether the channel proxy can change over the lifetime of the section */
 	mutable EMovieSceneChannelProxyType ChannelProxyType;
 };
+
+template<typename SectionParams>
+inline FFrameNumber GetFirstLoopStartOffsetAtTrimTime(FQualifiedFrameTime TrimTime, const SectionParams& Params, FFrameNumber StartFrame, FFrameRate FrameRate)
+{
+	const float AnimPlayRate = FMath::IsNearlyZero(Params.PlayRate) ? 1.0f : Params.PlayRate;
+	const float AnimPosition = (TrimTime.Time - StartFrame) / TrimTime.Rate * AnimPlayRate;
+	const float SeqLength = Params.GetSequenceLength() - FrameRate.AsSeconds(Params.StartFrameOffset + Params.EndFrameOffset) / AnimPlayRate;
+
+	FFrameNumber NewOffset = FrameRate.AsFrameNumber(FMath::Fmod(AnimPosition, SeqLength));
+	NewOffset += Params.FirstLoopStartFrameOffset;
+
+	const FFrameNumber SeqLengthInFrames = FrameRate.AsFrameNumber(SeqLength);
+	NewOffset = NewOffset % SeqLengthInFrames;
+
+	return NewOffset;
+}
diff --git a/Engine/Source/Runtime/MovieSceneTracks/Private/Evaluation/MovieSceneBaseCacheTemplate.cpp b/Engine/Source/Runtime/MovieSceneTracks/Private/Evaluation/MovieSceneBaseCacheTemplate.cpp
new file mode 100644
--- /dev/null
+++ b/Engine/Source/Runtime/MovieSceneTracks/Private/Evaluation/MovieSceneBaseCacheTemplate.cpp
@@ -0,0 +1,38 @@
+﻿// Copyright Epic Games, Inc. All Rights Reserved.
+
+#include "Evaluation/MovieSceneBaseCacheTemplate.h"
+
+DECLARE_CYCLE_STAT(TEXT("Chaos Cache Evaluate"), MovieSceneEval_BaseCache_Evaluate, STATGROUP_MovieSceneEval);
+DECLARE_CYCLE_STAT(TEXT("Chaos Cache Token Execute"), MovieSceneEval_BaseCache_TokenExecute, STATGROUP_MovieSceneEval);
+
+float FMovieSceneBaseCacheSectionTemplateParameters::MapTimeToAnimation(const FMovieSceneBaseCacheParams& BaseParams, float ComponentDuration, FFrameTime InPosition, FFrameRate InFrameRate) const
+{
+	const float SequenceLength = ComponentDuration;
+	const FFrameTime AnimationLength = SequenceLength * InFrameRate;
+	const int32 LengthInFrames = AnimationLength.FrameNumber.Value + (int)(AnimationLength.GetSubFrame() + 0.5f) + 1;
+	
+	//we only play end if we are not looping, and assuming we are looping if Length is greater than default length;
+	const bool bLooping = (SectionEndTime.Value - SectionStartTime.Value + BaseParams.StartFrameOffset + BaseParams.EndFrameOffset) > LengthInFrames;
+
+	InPosition = FMath::Clamp(InPosition, FFrameTime(SectionStartTime), FFrameTime(SectionEndTime - 1));
+
+	const float SectionPlayRate = BaseParams.PlayRate;
+	const float AnimPlayRate = FMath::IsNearlyZero(SectionPlayRate) ? 1.0f : SectionPlayRate;
+
+	const float FirstLoopSeqLength = SequenceLength - InFrameRate.AsSeconds(BaseParams.FirstLoopStartFrameOffset + BaseParams.StartFrameOffset + BaseParams.EndFrameOffset);
+	const float SeqLength = SequenceLength - InFrameRate.AsSeconds(BaseParams.StartFrameOffset + BaseParams.EndFrameOffset);
+
+	float AnimPosition = FFrameTime::FromDecimal((InPosition - SectionStartTime).AsDecimal() * AnimPlayRate) / InFrameRate;
+	AnimPosition += InFrameRate.AsSeconds(BaseParams.FirstLoopStartFrameOffset);
+	if (SeqLength > 0.f && (bLooping || !FMath::IsNearlyEqual(AnimPosition, SeqLength, 1e-4f)))
+	{
+		AnimPosition = FMath::Fmod(AnimPosition, SeqLength);
+	}
+	AnimPosition += InFrameRate.AsSeconds(BaseParams.StartFrameOffset);
+	if (BaseParams.bReverse)
+	{
+		AnimPosition = SequenceLength - AnimPosition;
+	}
+
+	return AnimPosition;
+}
diff --git a/Engine/Source/Runtime/MovieSceneTracks/Private/Sections/MovieSceneBaseCacheSection.cpp b/Engine/Source/Runtime/MovieSceneTracks/Private/Sections/MovieSceneBaseCacheSection.cpp
new file mode 100644
--- /dev/null
+++ b/Engine/Source/Runtime/MovieSceneTracks/Private/Sections/MovieSceneBaseCacheSection.cpp
@@ -0,0 +1,170 @@
+﻿// Copyright Epic Games, Inc. All Rights Reserved.
+
+#include "Sections/MovieSceneBaseCacheSection.h"
+#include "MovieScene.h"
+#include "Evaluation/MovieSceneBaseCacheTemplate.h"
+#include "MovieSceneTimeHelpers.h"
+#include "Misc/QualifiedFrameTime.h"
+
+FMovieSceneBaseCacheParams::FMovieSceneBaseCacheParams()
+{
+	PlayRate = 1.f;
+	bReverse = false;
+}
+
+UMovieSceneBaseCacheSection::UMovieSceneBaseCacheSection(const FObjectInitializer& ObjectInitializer)
+	: Super(ObjectInitializer), ParamsPtr(nullptr)
+{
+	BlendType = EMovieSceneBlendType::Absolute;
+	EvalOptions.EnableAndSetCompletionMode(EMovieSceneCompletionMode::ProjectDefault);
+
+#if WITH_EDITOR
+      	InitializePlayRate();
+      #endif
+}
+
+TOptional<FFrameTime> UMovieSceneBaseCacheSection::GetOffsetTime() const
+{
+	return ParamsPtr ? TOptional<FFrameTime>(ParamsPtr->FirstLoopStartFrameOffset) : TOptional<FFrameTime>(FFrameNumber());
+}
+
+TOptional<TRange<FFrameNumber> > UMovieSceneBaseCacheSection::GetAutoSizeRange() const
+{
+	if(ParamsPtr)
+	{
+		const FFrameRate FrameRate = GetTypedOuter<UMovieScene>()->GetTickResolution();
+		const FFrameTime AnimationLength = ParamsPtr->GetSequenceLength() * FrameRate;
+		const int32 IFrameNumber = AnimationLength.FrameNumber.Value + static_cast<int32>(AnimationLength.GetSubFrame() + 0.5f);
+
+		return TRange<FFrameNumber>(GetInclusiveStartFrame(), GetInclusiveStartFrame() + IFrameNumber + 1);
+	}
+	return TRange<FFrameNumber>(0, 1);
+}
+
+void UMovieSceneBaseCacheSection::TrimSection(FQualifiedFrameTime TrimTime, bool bTrimLeft, bool bDeleteKeys)
+{
+	SetFlags(RF_Transactional);
+
+	if (TryModify() && ParamsPtr)
+	{
+		if (bTrimLeft)
+		{
+			FFrameRate FrameRate = GetTypedOuter<UMovieScene>()->GetTickResolution();
+
+			ParamsPtr->FirstLoopStartFrameOffset = HasStartFrame() ? GetFirstLoopStartOffsetAtTrimTime(TrimTime, *ParamsPtr, GetInclusiveStartFrame(), FrameRate) : 0;
+		}
+
+		Super::TrimSection(TrimTime, bTrimLeft, bDeleteKeys);
+	}
+}
+
+UMovieSceneSection* UMovieSceneBaseCacheSection::SplitSection(FQualifiedFrameTime SplitTime, bool bDeleteKeys)
+{
+	if(ParamsPtr)
+	{
+		const FFrameNumber InitialFirstLoopStartFrameOffset = ParamsPtr->FirstLoopStartFrameOffset;
+
+		FFrameRate FrameRate = GetTypedOuter<UMovieScene>()->GetTickResolution();
+
+		const FFrameNumber NewOffset = HasStartFrame() ? GetFirstLoopStartOffsetAtTrimTime(SplitTime, *ParamsPtr, GetInclusiveStartFrame(), FrameRate) : 0;
+
+		UMovieSceneSection* NewSection = Super::SplitSection(SplitTime, bDeleteKeys);
+		if (NewSection != nullptr)
+		{
+			UMovieSceneBaseCacheSection* NewChaosCacheSection = Cast<UMovieSceneBaseCacheSection>(NewSection);
+			NewChaosCacheSection->ParamsPtr->FirstLoopStartFrameOffset = NewOffset;
+		}
+
+		// Restore original offset modified by splitting
+		ParamsPtr->FirstLoopStartFrameOffset = InitialFirstLoopStartFrameOffset;
+
+		return NewSection;
+	}
+	return nullptr;
+}
+
+void UMovieSceneBaseCacheSection::GetSnapTimes(TArray<FFrameNumber>& OutSnapTimes, bool bGetSectionBorders) const
+{
+	Super::GetSnapTimes(OutSnapTimes, bGetSectionBorders);
+	if(ParamsPtr)
+	{
+		const FFrameRate   FrameRate = GetTypedOuter<UMovieScene>()->GetTickResolution();
+		const FFrameNumber StartFrame = GetInclusiveStartFrame();
+		const FFrameNumber EndFrame = GetExclusiveEndFrame() - 1; // -1 because we don't need to add the end frame twice
+
+		const float AnimPlayRate = FMath::IsNearlyZero(ParamsPtr->PlayRate) ? 1.0f : ParamsPtr->PlayRate;
+		const float SeqLengthSeconds = ParamsPtr->GetSequenceLength() - FrameRate.AsSeconds(ParamsPtr->StartFrameOffset + ParamsPtr->EndFrameOffset) / AnimPlayRate;
+		const float FirstLoopSeqLengthInSeconds = SeqLengthSeconds - FrameRate.AsSeconds(ParamsPtr->FirstLoopStartFrameOffset) / AnimPlayRate;
+
+		const FFrameTime SequenceFrameLength = SeqLengthSeconds * FrameRate;
+		const FFrameTime FirstLoopSequenceFrameLength = FirstLoopSeqLengthInSeconds * FrameRate;
+		if (SequenceFrameLength.FrameNumber > 1)
+		{
+			// Snap to the repeat times
+			bool IsFirstLoop = true;
+			FFrameTime CurrentTime = StartFrame;
+			while (CurrentTime < EndFrame)
+			{
+				OutSnapTimes.Add(CurrentTime.FrameNumber);
+				if (IsFirstLoop)
+				{
+					CurrentTime += FirstLoopSequenceFrameLength;
+					IsFirstLoop = false;
+				}
+				else
+				{
+					CurrentTime += SequenceFrameLength;
+				}
+			}
+		}
+	}
+}
+
+float UMovieSceneBaseCacheSection::MapTimeToAnimation(float ComponentDuration, FFrameTime InPosition, FFrameRate InFrameRate) const
+{
+	if(ParamsPtr)
+	{
+		const FMovieSceneBaseCacheSectionTemplateParameters TemplateParamsPtr(GetInclusiveStartFrame(), GetExclusiveEndFrame());
+		return TemplateParamsPtr.MapTimeToAnimation(*ParamsPtr, ComponentDuration, InPosition, InFrameRate);
+	}
+	return 0.0f;
+}
+
+#if WITH_EDITOR
+void UMovieSceneBaseCacheSection::PreEditChange(FProperty* PropertyAboutToChange)
+{
+	if(ParamsPtr)
+	{
+		// Store the current play rate so that we can compute the amount to compensate the section end time when the play rate changes
+		PreviousPlayRate = ParamsPtr->PlayRate;
+	}
+	Super::PreEditChange(PropertyAboutToChange);
+}
+
+void UMovieSceneBaseCacheSection::PostEditChangeProperty(FPropertyChangedEvent& PropertyChangedEvent)
+{
+	// Adjust the duration automatically if the play rate changes
+	if (ParamsPtr && PropertyChangedEvent.Property != nullptr &&
+		PropertyChangedEvent.Property->GetFName() == TEXT("PlayRate"))
+	{
+		const float NewPlayRate = ParamsPtr->PlayRate;
+
+		if (!FMath::IsNearlyZero(NewPlayRate))
+		{
+			const float CurrentDuration = UE::MovieScene::DiscreteSize(GetRange());
+			const float NewDuration = CurrentDuration * (PreviousPlayRate / NewPlayRate);
+			SetEndFrame(GetInclusiveStartFrame() + FMath::FloorToInt(NewDuration));
+
+			PreviousPlayRate = NewPlayRate;
+		}
+	}
+	Super::PostEditChangeProperty(PropertyChangedEvent);
+}
+
+void UMovieSceneBaseCacheSection::InitializePlayRate()
+{
+	PreviousPlayRate = ParamsPtr ? ParamsPtr->PlayRate : 0.0f;
+}
+
+#endif
+
diff --git a/Engine/Source/Runtime/MovieSceneTracks/Private/Sections/MovieSceneSkeletalAnimationSection.cpp b/Engine/Source/Runtime/MovieSceneTracks/Private/Sections/MovieSceneSkeletalAnimationSection.cpp
--- a/Engine/Source/Runtime/MovieSceneTracks/Private/Sections/MovieSceneSkeletalAnimationSection.cpp
+++ b/Engine/Source/Runtime/MovieSceneTracks/Private/Sections/MovieSceneSkeletalAnimationSection.cpp
@@ -213,23 +213,6 @@ void UMovieSceneSkeletalAnimationSection::PostLoad()
 	Super::PostLoad();
 }
 
-FFrameNumber GetFirstLoopStartOffsetAtTrimTime(FQualifiedFrameTime TrimTime, const FMovieSceneSkeletalAnimationParams& Params, FFrameNumber StartFrame, FFrameRate FrameRate)
-{
-	const float AnimPlayRate = FMath::IsNearlyZero(Params.PlayRate) ? 1.0f : Params.PlayRate;
-	const float AnimPosition = (TrimTime.Time - StartFrame) / TrimTime.Rate * AnimPlayRate;
-	const float SeqLength = Params.GetSequenceLength() - FrameRate.AsSeconds(Params.StartFrameOffset + Params.EndFrameOffset) / AnimPlayRate;
-
-	FFrameNumber NewOffset = FrameRate.AsFrameNumber(FMath::Fmod(AnimPosition, SeqLength));
-	NewOffset += Params.FirstLoopStartFrameOffset;
-
-	const FFrameNumber SeqLengthInFrames = FrameRate.AsFrameNumber(SeqLength);
-	while (NewOffset >= SeqLengthInFrames)
-		NewOffset -= SeqLengthInFrames;
-
-	return NewOffset;
-}
-
-
 TOptional<TRange<FFrameNumber> > UMovieSceneSkeletalAnimationSection::GetAutoSizeRange() const
 {
 	FFrameRate FrameRate = GetTypedOuter<UMovieScene>()->GetTickResolution();
diff --git a/Engine/Source/Runtime/MovieSceneTracks/Public/Evaluation/MovieSceneBaseCacheTemplate.h b/Engine/Source/Runtime/MovieSceneTracks/Public/Evaluation/MovieSceneBaseCacheTemplate.h
new file mode 100644
--- /dev/null
+++ b/Engine/Source/Runtime/MovieSceneTracks/Public/Evaluation/MovieSceneBaseCacheTemplate.h
@@ -0,0 +1,30 @@
+﻿// Copyright Epic Games, Inc. All Rights Reserved.
+
+#pragma once
+
+#include "Evaluation/MovieSceneEvalTemplate.h"
+#include "Sections/MovieSceneBaseCacheSection.h"
+#include "MovieSceneBaseCacheTemplate.generated.h"
+
+/** Base cache parameters that will be used for all the base cache sections */
+USTRUCT()
+struct FMovieSceneBaseCacheSectionTemplateParameters
+{
+	GENERATED_BODY()
+	FMovieSceneBaseCacheSectionTemplateParameters() {}
+	FMovieSceneBaseCacheSectionTemplateParameters(FFrameNumber InSectionStartTime, FFrameNumber InSectionEndTime)
+		: SectionStartTime(InSectionStartTime)
+		, SectionEndTime(InSectionEndTime)
+	{}
+
+	/** Get Frame Time as Animation Time */
+	MOVIESCENETRACKS_API float MapTimeToAnimation(const FMovieSceneBaseCacheParams& BaseParams, float ComponentDuration, FFrameTime InPosition, FFrameRate InFrameRate) const;
+
+	/** Section start time */
+	UPROPERTY()
+	FFrameNumber SectionStartTime;
+
+	/** Section end time */
+	UPROPERTY()
+	FFrameNumber SectionEndTime;
+};
diff --git a/Engine/Source/Runtime/MovieSceneTracks/Public/Sections/MovieSceneBaseCacheSection.h b/Engine/Source/Runtime/MovieSceneTracks/Public/Sections/MovieSceneBaseCacheSection.h
new file mode 100644
--- /dev/null
+++ b/Engine/Source/Runtime/MovieSceneTracks/Public/Sections/MovieSceneBaseCacheSection.h
@@ -0,0 +1,78 @@
+﻿// Copyright Epic Games, Inc. All Rights Reserved.
+
+#pragma once
+
+#include "MovieSceneSection.h"
+#include "MovieSceneBaseCacheSection.generated.h"
+
+/**
+ * Base class for the cache parameters that will be used in all the cache sections
+ */
+USTRUCT()
+struct MOVIESCENETRACKS_API FMovieSceneBaseCacheParams
+{
+	GENERATED_BODY()
+
+	/** Constructor/Destructor */
+	FMovieSceneBaseCacheParams();
+	virtual ~FMovieSceneBaseCacheParams() {}
+
+	/** Gets the animation sequence length, not modified by play rate */
+	virtual float GetSequenceLength() const {return 0.0f;}
+
+	/** The offset for the first loop of the animation clip */
+	UPROPERTY(EditAnywhere, Category = "Cache")
+	FFrameNumber FirstLoopStartFrameOffset;
+
+	/** The offset into the beginning of the animation clip */
+	UPROPERTY(EditAnywhere, Category = "Cache")
+	FFrameNumber StartFrameOffset;
+
+	/** The offset into the end of the animation clip */
+	UPROPERTY(EditAnywhere, Category = "Cache")
+	FFrameNumber EndFrameOffset;
+
+	/** The playback rate of the animation clip */
+	UPROPERTY(EditAnywhere, Category = "Cache")
+	float PlayRate;
+
+	/** Reverse the playback of the animation clip */
+	UPROPERTY(EditAnywhere, Category = "Cache")
+	uint32 bReverse : 1;
+};
+
+/**
+ * Movie scene section that control base cache playback
+ */
+UCLASS(MinimalAPI)
+class  UMovieSceneBaseCacheSection
+	: public UMovieSceneSection
+{
+	GENERATED_UCLASS_BODY()
+
+public:
+	/** Pointer to the concrete uproperty that will be instanced in each BaseCacheSection children */  
+	FMovieSceneBaseCacheParams* ParamsPtr = nullptr;
+	
+	/** Get Frame Time as Animation Time*/
+	MOVIESCENETRACKS_API virtual float MapTimeToAnimation(float ComponentDuration, FFrameTime InPosition, FFrameRate InFrameRate) const;
+	MOVIESCENETRACKS_API virtual TOptional<TRange<FFrameNumber> > GetAutoSizeRange() const override;
+	
+protected:
+	//~ UMovieSceneSection interface
+	MOVIESCENETRACKS_API virtual void TrimSection(FQualifiedFrameTime TrimTime, bool bTrimLeft, bool bDeleteKeys) override;
+	MOVIESCENETRACKS_API virtual UMovieSceneSection* SplitSection(FQualifiedFrameTime SplitTime, bool bDeleteKeys) override;
+	MOVIESCENETRACKS_API virtual void GetSnapTimes(TArray<FFrameNumber>& OutSnapTimes, bool bGetSectionBorders) const override;
+	MOVIESCENETRACKS_API virtual TOptional<FFrameTime> GetOffsetTime() const override;
+
+#if WITH_EDITOR
+	MOVIESCENETRACKS_API void InitializePlayRate();
+
+private:
+	MOVIESCENETRACKS_API virtual void PreEditChange(FProperty* PropertyAboutToChange) override;
+	MOVIESCENETRACKS_API virtual void PostEditChangeProperty(FPropertyChangedEvent& PropertyChangedEvent) override;
+public:
+	float PreviousPlayRate;
+
+#endif
+};
