Fix for a compile issue post-merge after taking a change to FComponentReference for a KPK issue.

diff --git a/Engine/Plugins/Experimental/ChaosCaching/Source/ChaosCaching/Private/Chaos/CacheManagerActor.cpp b/Engine/Plugins/Experimental/ChaosCaching/Source/ChaosCaching/Private/Chaos/CacheManagerActor.cpp
--- a/Engine/Plugins/Experimental/ChaosCaching/Source/ChaosCaching/Private/Chaos/CacheManagerActor.cpp
+++ b/Engine/Plugins/Experimental/ChaosCaching/Source/ChaosCaching/Private/Chaos/CacheManagerActor.cpp
@@ -44,9 +44,9 @@ void FObservedComponent::ResetRuntimeData(const EStartMode ManagerStartMode)
 void FObservedComponent::PostSerialize(const FArchive& Ar)
 {
 	PRAGMA_DISABLE_DEPRECATION_WARNINGS
-	if (Ar.IsLoading() && ComponentRef.OtherActor)
+	if (Ar.IsLoading() && ComponentRef.OtherActor.IsValid())
 	{
-		SoftComponentRef.OtherActor = ComponentRef.OtherActor;
+		SoftComponentRef.OtherActor = ComponentRef.OtherActor.Get();
 		SoftComponentRef.ComponentProperty = ComponentRef.ComponentProperty;
 		SoftComponentRef.OverrideComponent = ComponentRef.OverrideComponent;
 		SoftComponentRef.PathToComponent = ComponentRef.PathToComponent; 
