Fix chaos cache when no collection is set

diff --git a/Engine/Plugins/Experimental/ChaosCaching/Source/ChaosCaching/Private/Chaos/Sequencer/MovieSceneChaosCacheTemplate.cpp b/Engine/Plugins/Experimental/ChaosCaching/Source/ChaosCaching/Private/Chaos/Sequencer/MovieSceneChaosCacheTemplate.cpp
--- a/Engine/Plugins/Experimental/ChaosCaching/Source/ChaosCaching/Private/Chaos/Sequencer/MovieSceneChaosCacheTemplate.cpp
+++ b/Engine/Plugins/Experimental/ChaosCaching/Source/ChaosCaching/Private/Chaos/Sequencer/MovieSceneChaosCacheTemplate.cpp
@@ -71,7 +71,7 @@ struct FChaosCacheExecutionToken : IMovieSceneExecutionToken
 				if (UObject* Obj = WeakObj.Get())
 				{
 					AChaosCacheManager*  ChaosCache = ChaosCacheFromObject(Obj);
-					if (ChaosCache)
+					if (ChaosCache && ChaosCache->CacheCollection)
 					{ 
 						// Set the ChaosCache on the component only if it's set and valid in the Params
 						if (Params.ChaosCacheParams.CacheCollection && Params.ChaosCacheParams.CacheCollection != ChaosCache->CacheCollection)
diff --git a/Engine/Plugins/Experimental/ChaosCaching/Source/ChaosCaching/Private/Chaos/Sequencer/MovieSceneChaosCacheTrack.cpp b/Engine/Plugins/Experimental/ChaosCaching/Source/ChaosCaching/Private/Chaos/Sequencer/MovieSceneChaosCacheTrack.cpp
--- a/Engine/Plugins/Experimental/ChaosCaching/Source/ChaosCaching/Private/Chaos/Sequencer/MovieSceneChaosCacheTrack.cpp
+++ b/Engine/Plugins/Experimental/ChaosCaching/Source/ChaosCaching/Private/Chaos/Sequencer/MovieSceneChaosCacheTrack.cpp
@@ -32,11 +32,14 @@ UMovieSceneSection* UMovieSceneChaosCacheTrack::AddNewAnimation(FFrameNumber Key
 {
 	UMovieSceneChaosCacheSection* NewSection = Cast<UMovieSceneChaosCacheSection>(CreateNewSection());
 	{
-		const FFrameTime AnimationLength = ChaosCache->CacheCollection->GetMaxDuration() * GetTypedOuter<UMovieScene>()->GetTickResolution();
-		const int32 IFrameNumber = AnimationLength.FrameNumber.Value + (int)(AnimationLength.GetSubFrame() + 0.5f) + 1;
-		NewSection->InitialPlacementOnRow(AnimationSections, KeyTime, IFrameNumber, INDEX_NONE);
-		
-		NewSection->Params.CacheCollection = ChaosCache->CacheCollection;
+		if(ChaosCache && ChaosCache->CacheCollection)
+		{
+			const FFrameTime AnimationLength = ChaosCache->CacheCollection->GetMaxDuration() * GetTypedOuter<UMovieScene>()->GetTickResolution();
+			const int32 IFrameNumber = AnimationLength.FrameNumber.Value + (int)(AnimationLength.GetSubFrame() + 0.5f) + 1;
+			NewSection->InitialPlacementOnRow(AnimationSections, KeyTime, IFrameNumber, INDEX_NONE);
+			
+			NewSection->Params.CacheCollection = ChaosCache->CacheCollection;
+		}
 	}
 
 	AddSection(*NewSection);
