Updated to use UE_INLINE_GENERATED_CPP_BY_NAME.

diff --git a/Engine/Plugins/Experimental/Animation/MotionTrajectory/Source/MotionTrajectory/Private/MotionTrajectory.cpp b/Engine/Plugins/Experimental/Animation/MotionTrajectory/Source/MotionTrajectory/Private/MotionTrajectory.cpp
--- a/Engine/Plugins/Experimental/Animation/MotionTrajectory/Source/MotionTrajectory/Private/MotionTrajectory.cpp
+++ b/Engine/Plugins/Experimental/Animation/MotionTrajectory/Source/MotionTrajectory/Private/MotionTrajectory.cpp
@@ -7,6 +7,8 @@
 #include "Kismet/KismetMathLibrary.h"
 #include "Kismet/KismetSystemLibrary.h"
 
+#include UE_INLINE_GENERATED_CPP_BY_NAME(MotionTrajectory)
+
 #define LOCTEXT_NAMESPACE "MotionTrajectory"
 
 DEFINE_LOG_CATEGORY(LogMotionTrajectory);
@@ -288,4 +290,4 @@ FTrajectorySampleRange UMotionTrajectoryComponent::GetHistory() const
 
 #undef LOCTEXT_NAMESPACE
 
-IMPLEMENT_MODULE(FMotionTrajectoryModule, MotionTrajectory)
\ No newline at end of file
+IMPLEMENT_MODULE(FMotionTrajectoryModule, MotionTrajectory)
diff --git a/Engine/Plugins/Experimental/Animation/MotionTrajectory/Source/MotionTrajectory/Private/MotionTrajectoryCharacterMovement.cpp b/Engine/Plugins/Experimental/Animation/MotionTrajectory/Source/MotionTrajectory/Private/MotionTrajectoryCharacterMovement.cpp
--- a/Engine/Plugins/Experimental/Animation/MotionTrajectory/Source/MotionTrajectory/Private/MotionTrajectoryCharacterMovement.cpp
+++ b/Engine/Plugins/Experimental/Animation/MotionTrajectory/Source/MotionTrajectory/Private/MotionTrajectoryCharacterMovement.cpp
@@ -9,6 +9,8 @@
 #include "GameFramework/Controller.h"
 #include "KismetAnimationLibrary.h"
 
+#include UE_INLINE_GENERATED_CPP_BY_NAME(MotionTrajectoryCharacterMovement)
+
 FTrajectorySample UCharacterMovementTrajectoryComponent::CalcWorldSpacePresentTrajectorySample(float DeltaTime) const
 {
 	FTrajectorySample ReturnValue;
diff --git a/Engine/Plugins/Experimental/Animation/MotionTrajectory/Source/MotionTrajectory/Private/MotionTrajectoryLibrary.cpp b/Engine/Plugins/Experimental/Animation/MotionTrajectory/Source/MotionTrajectory/Private/MotionTrajectoryLibrary.cpp
--- a/Engine/Plugins/Experimental/Animation/MotionTrajectory/Source/MotionTrajectory/Private/MotionTrajectoryLibrary.cpp
+++ b/Engine/Plugins/Experimental/Animation/MotionTrajectory/Source/MotionTrajectory/Private/MotionTrajectoryLibrary.cpp
@@ -7,6 +7,8 @@
 #include "Algo/Find.h"
 #include "GameFramework/Actor.h"
 
+#include UE_INLINE_GENERATED_CPP_BY_NAME(MotionTrajectoryLibrary)
+
 static void FlattenTrajectoryPosition(
 	FTrajectorySample& Sample, 
 	const FTrajectorySample& PrevSample, 
diff --git a/Engine/Plugins/Experimental/Animation/PoseSearch/Source/Runtime/Private/AnimNode_MotionMatching.cpp b/Engine/Plugins/Experimental/Animation/PoseSearch/Source/Runtime/Private/AnimNode_MotionMatching.cpp
--- a/Engine/Plugins/Experimental/Animation/PoseSearch/Source/Runtime/Private/AnimNode_MotionMatching.cpp
+++ b/Engine/Plugins/Experimental/Animation/PoseSearch/Source/Runtime/Private/AnimNode_MotionMatching.cpp
@@ -14,6 +14,8 @@
 #include "PoseSearch/PoseSearchDerivedData.h"
 #include "Trace/PoseSearchTraceLogger.h"
 
+#include UE_INLINE_GENERATED_CPP_BY_NAME(AnimNode_MotionMatching)
+
 #define LOCTEXT_NAMESPACE "AnimNode_MotionMatching"
 
 /////////////////////////////////////////////////////
@@ -258,4 +260,4 @@ bool FAnimNode_MotionMatching::SetIgnoreForRelevancyTest(bool bInIgnoreForReleva
 	return false;
 }
 
-#undef LOCTEXT_NAMESPACE
\ No newline at end of file
+#undef LOCTEXT_NAMESPACE
diff --git a/Engine/Plugins/Experimental/Animation/PoseSearch/Source/Runtime/Private/AnimNode_PoseSearchHistoryCollector.cpp b/Engine/Plugins/Experimental/Animation/PoseSearch/Source/Runtime/Private/AnimNode_PoseSearchHistoryCollector.cpp
--- a/Engine/Plugins/Experimental/Animation/PoseSearch/Source/Runtime/Private/AnimNode_PoseSearchHistoryCollector.cpp
+++ b/Engine/Plugins/Experimental/Animation/PoseSearch/Source/Runtime/Private/AnimNode_PoseSearchHistoryCollector.cpp
@@ -4,6 +4,8 @@
 #include "Animation/AnimInstanceProxy.h"
 #include "Animation/AnimNodeMessages.h"
 
+#include UE_INLINE_GENERATED_CPP_BY_NAME(AnimNode_PoseSearchHistoryCollector)
+
 #define LOCTEXT_NAMESPACE "AnimNode_PoseSearchHistoryCollector"
 
 namespace UE::PoseSearch::Private
@@ -95,4 +97,4 @@ void FAnimNode_PoseSearchHistoryCollector::GatherDebugData(FNodeDebugData& Debug
 	Source.GatherDebugData(DebugData);
 }
 
-#undef LOCTEXT_NAMESPACE
\ No newline at end of file
+#undef LOCTEXT_NAMESPACE
diff --git a/Engine/Plugins/Experimental/Animation/PoseSearch/Source/Runtime/Private/PoseSearch.cpp b/Engine/Plugins/Experimental/Animation/PoseSearch/Source/Runtime/Private/PoseSearch.cpp
--- a/Engine/Plugins/Experimental/Animation/PoseSearch/Source/Runtime/Private/PoseSearch.cpp
+++ b/Engine/Plugins/Experimental/Animation/PoseSearch/Source/Runtime/Private/PoseSearch.cpp
@@ -31,6 +31,8 @@
 #include "Animation/AnimRootMotionProvider.h"
 #include "Animation/BlendSpace1D.h"
 
+
+#include UE_INLINE_GENERATED_CPP_BY_NAME(PoseSearch)
 #if WITH_EDITOR
 #include "DerivedDataRequestOwner.h"
 #endif
@@ -2966,4 +2968,4 @@ public: // IModuleInterface
 
 #undef LOCTEXT_NAMESPACE
 
-IMPLEMENT_MODULE(UE::PoseSearch::FModule, PoseSearch)
\ No newline at end of file
+IMPLEMENT_MODULE(UE::PoseSearch::FModule, PoseSearch)
diff --git a/Engine/Plugins/Experimental/Animation/PoseSearch/Source/Runtime/Private/PoseSearchAnimNotifies.cpp b/Engine/Plugins/Experimental/Animation/PoseSearch/Source/Runtime/Private/PoseSearchAnimNotifies.cpp
--- a/Engine/Plugins/Experimental/Animation/PoseSearch/Source/Runtime/Private/PoseSearchAnimNotifies.cpp
+++ b/Engine/Plugins/Experimental/Animation/PoseSearch/Source/Runtime/Private/PoseSearchAnimNotifies.cpp
@@ -1,3 +1,6 @@
 // Copyright Epic Games, Inc. All Rights Reserved.
 
-#include "PoseSearch/PoseSearchAnimNotifies.h"
\ No newline at end of file
+#include "PoseSearch/PoseSearchAnimNotifies.h"
+
+#include UE_INLINE_GENERATED_CPP_BY_NAME(PoseSearchAnimNotifies)
+
diff --git a/Engine/Plugins/Experimental/Animation/PoseSearch/Source/Runtime/Private/PoseSearchDerivedData.cpp b/Engine/Plugins/Experimental/Animation/PoseSearch/Source/Runtime/Private/PoseSearchDerivedData.cpp
--- a/Engine/Plugins/Experimental/Animation/PoseSearch/Source/Runtime/Private/PoseSearchDerivedData.cpp
+++ b/Engine/Plugins/Experimental/Animation/PoseSearch/Source/Runtime/Private/PoseSearchDerivedData.cpp
@@ -2,6 +2,8 @@
 
 #include "PoseSearch/PoseSearchDerivedData.h"
 
+#include UE_INLINE_GENERATED_CPP_BY_NAME(PoseSearchDerivedData)
+
 #if WITH_EDITOR
 #include "Animation/BlendSpace.h"
 #include "AssetRegistry/ARFilter.h"
diff --git a/Engine/Plugins/Experimental/Animation/PoseSearch/Source/Runtime/Private/PoseSearchFeatureChannels.cpp b/Engine/Plugins/Experimental/Animation/PoseSearch/Source/Runtime/Private/PoseSearchFeatureChannels.cpp
--- a/Engine/Plugins/Experimental/Animation/PoseSearch/Source/Runtime/Private/PoseSearchFeatureChannels.cpp
+++ b/Engine/Plugins/Experimental/Animation/PoseSearch/Source/Runtime/Private/PoseSearchFeatureChannels.cpp
@@ -10,6 +10,8 @@
 #include "PoseSearchEigenHelper.h"
 #include "UObject/ObjectSaveContext.h"
 
+#include UE_INLINE_GENERATED_CPP_BY_NAME(PoseSearchFeatureChannels)
+
 #define LOCTEXT_NAMESPACE "PoseSearchFeatureChannels"
 
 namespace UE::PoseSearch
@@ -2353,4 +2355,4 @@ bool UPoseSearchFeatureChannel_FilterCrashingLegs::IsPoseValid(TConstArrayView<f
 	return bIsPoseValid;
 }
 
-#undef LOCTEXT_NAMESPACE
\ No newline at end of file
+#undef LOCTEXT_NAMESPACE
diff --git a/Engine/Plugins/Experimental/Animation/PoseSearch/Source/Runtime/Private/PoseSearchLibrary.cpp b/Engine/Plugins/Experimental/Animation/PoseSearch/Source/Runtime/Private/PoseSearchLibrary.cpp
--- a/Engine/Plugins/Experimental/Animation/PoseSearch/Source/Runtime/Private/PoseSearchLibrary.cpp
+++ b/Engine/Plugins/Experimental/Animation/PoseSearch/Source/Runtime/Private/PoseSearchLibrary.cpp
@@ -13,6 +13,8 @@
 #include "MotionTrajectoryLibrary.h"
 #include "PoseSearch/PoseSearchFeatureChannels.h"
 
+#include UE_INLINE_GENERATED_CPP_BY_NAME(PoseSearchLibrary)
+
 #define LOCTEXT_NAMESPACE "PoseSearchLibrary"
 
 //////////////////////////////////////////////////////////////////////////
@@ -469,4 +471,4 @@ float FMotionMatchingState::ComputeJumpBlendTime(const UE::PoseSearch::FSearchRe
 	return JumpBlendTime;
 }
 
-#undef LOCTEXT_NAMESPACE
\ No newline at end of file
+#undef LOCTEXT_NAMESPACE
diff --git a/Engine/Plugins/Experimental/ChaosCaching/Source/ChaosCaching/Private/Chaos/Adapters/GeometryCollectionComponentCacheAdapter.cpp b/Engine/Plugins/Experimental/ChaosCaching/Source/ChaosCaching/Private/Chaos/Adapters/GeometryCollectionComponentCacheAdapter.cpp
--- a/Engine/Plugins/Experimental/ChaosCaching/Source/ChaosCaching/Private/Chaos/Adapters/GeometryCollectionComponentCacheAdapter.cpp
+++ b/Engine/Plugins/Experimental/ChaosCaching/Source/ChaosCaching/Private/Chaos/Adapters/GeometryCollectionComponentCacheAdapter.cpp
@@ -12,6 +12,8 @@
 #include "Chaos/ChaosSolverActor.h"
 #include "GeometryCollection/GeometryCollectionAlgo.h"
 
+#include UE_INLINE_GENERATED_CPP_BY_NAME(GeometryCollectionComponentCacheAdapter)
+
 FName FEnableStateEvent::EventName("GC_Enable");
 FName FBreakingEvent::EventName("GC_Breaking");
 FName FCollisionEvent::EventName("GC_Collision");
@@ -902,3 +904,4 @@ namespace Chaos
 	}
 
 }    // namespace Chaos
+
diff --git a/Engine/Plugins/Experimental/ChaosCaching/Source/ChaosCaching/Private/Chaos/CacheCollection.cpp b/Engine/Plugins/Experimental/ChaosCaching/Source/ChaosCaching/Private/Chaos/CacheCollection.cpp
--- a/Engine/Plugins/Experimental/ChaosCaching/Source/ChaosCaching/Private/Chaos/CacheCollection.cpp
+++ b/Engine/Plugins/Experimental/ChaosCaching/Source/ChaosCaching/Private/Chaos/CacheCollection.cpp
@@ -5,6 +5,8 @@
 #include "Algo/Find.h"
 #include "Async/ParallelFor.h"
 
+#include UE_INLINE_GENERATED_CPP_BY_NAME(CacheCollection)
+
 void UChaosCacheCollection::GetAssetRegistryTags(TArray<FAssetRegistryTag>& OutTags) const
 {
 #if WITH_EDITOR
@@ -95,4 +97,4 @@ float UChaosCacheCollection::GetMaxDuration() const
     	MaxDuration = FMath::Max(CacheInstance->GetDuration(), MaxDuration);
     }
 	return MaxDuration;
-}
\ No newline at end of file
+}
diff --git a/Engine/Plugins/Experimental/ChaosCaching/Source/ChaosCaching/Private/Chaos/CacheEvents.cpp b/Engine/Plugins/Experimental/ChaosCaching/Source/ChaosCaching/Private/Chaos/CacheEvents.cpp
--- a/Engine/Plugins/Experimental/ChaosCaching/Source/ChaosCaching/Private/Chaos/CacheEvents.cpp
+++ b/Engine/Plugins/Experimental/ChaosCaching/Source/ChaosCaching/Private/Chaos/CacheEvents.cpp
@@ -3,6 +3,8 @@
 #include "Chaos/CacheEvents.h"
 #include "Logging/MessageLog.h"
 
+#include UE_INLINE_GENERATED_CPP_BY_NAME(CacheEvents)
+
 FCacheEventTrack::FCacheEventTrack()
 	: Name(NAME_None)
 	, Struct(nullptr)
@@ -180,3 +182,4 @@ bool FCacheEventTrack::FHandle::IsAlive() const
 {
 	return Index != INDEX_NONE && Track && Track->GetTransientVersion() == Version;
 }
+
diff --git a/Engine/Plugins/Experimental/ChaosCaching/Source/ChaosCaching/Private/Chaos/CacheManagerActor.cpp b/Engine/Plugins/Experimental/ChaosCaching/Source/ChaosCaching/Private/Chaos/CacheManagerActor.cpp
--- a/Engine/Plugins/Experimental/ChaosCaching/Source/ChaosCaching/Private/Chaos/CacheManagerActor.cpp
+++ b/Engine/Plugins/Experimental/ChaosCaching/Source/ChaosCaching/Private/Chaos/CacheManagerActor.cpp
@@ -14,6 +14,8 @@
 #include "Features/IModularFeatures.h"
 #include "Algo/Find.h"
 
+#include UE_INLINE_GENERATED_CPP_BY_NAME(CacheManagerActor)
+
 #if WITH_EDITOR
 #include "Editor.h"
 #include "Editor/EditorEngine.h"
@@ -754,3 +756,4 @@ AChaosCachePlayer::AChaosCachePlayer(const FObjectInitializer& ObjectInitializer
 }
 
 #undef LOCTEXT_NAMESPACE
+
diff --git a/Engine/Plugins/Experimental/ChaosCaching/Source/ChaosCaching/Private/Chaos/ChaosCache.cpp b/Engine/Plugins/Experimental/ChaosCaching/Source/ChaosCaching/Private/Chaos/ChaosCache.cpp
--- a/Engine/Plugins/Experimental/ChaosCaching/Source/ChaosCaching/Private/Chaos/ChaosCache.cpp
+++ b/Engine/Plugins/Experimental/ChaosCaching/Source/ChaosCaching/Private/Chaos/ChaosCache.cpp
@@ -9,6 +9,8 @@
 #include "GeometryCollection/GeometryCollectionComponent.h"
 #include "GeometryCollection/GeometryCollectionObject.h"
 
+#include UE_INLINE_GENERATED_CPP_BY_NAME(ChaosCache)
+
 bool bChaosCacheUseInterpolation = true;
 FAutoConsoleVariableRef CVarChaosCacheUseInterpolation(
 	TEXT("p.Chaos.Cache.UseInterpolation"),
diff --git a/Engine/Plugins/Experimental/ChaosCaching/Source/ChaosCaching/Private/Chaos/Sequencer/MovieSceneChaosCacheSection.cpp b/Engine/Plugins/Experimental/ChaosCaching/Source/ChaosCaching/Private/Chaos/Sequencer/MovieSceneChaosCacheSection.cpp
--- a/Engine/Plugins/Experimental/ChaosCaching/Source/ChaosCaching/Private/Chaos/Sequencer/MovieSceneChaosCacheSection.cpp
+++ b/Engine/Plugins/Experimental/ChaosCaching/Source/ChaosCaching/Private/Chaos/Sequencer/MovieSceneChaosCacheSection.cpp
@@ -4,6 +4,8 @@
 #include "Chaos/CacheCollection.h"
 #include "MovieScene.h"
 
+#include UE_INLINE_GENERATED_CPP_BY_NAME(MovieSceneChaosCacheSection)
+
 FMovieSceneChaosCacheParams::FMovieSceneChaosCacheParams() : FMovieSceneBaseCacheParams()
 {
 	CacheCollection = nullptr;
diff --git a/Engine/Plugins/Experimental/ChaosCaching/Source/ChaosCaching/Private/Chaos/Sequencer/MovieSceneChaosCacheTemplate.cpp b/Engine/Plugins/Experimental/ChaosCaching/Source/ChaosCaching/Private/Chaos/Sequencer/MovieSceneChaosCacheTemplate.cpp
--- a/Engine/Plugins/Experimental/ChaosCaching/Source/ChaosCaching/Private/Chaos/Sequencer/MovieSceneChaosCacheTemplate.cpp
+++ b/Engine/Plugins/Experimental/ChaosCaching/Source/ChaosCaching/Private/Chaos/Sequencer/MovieSceneChaosCacheTemplate.cpp
@@ -5,6 +5,8 @@
 #include "Chaos/CacheCollection.h"
 #include "Chaos/CacheManagerActor.h"
 
+#include UE_INLINE_GENERATED_CPP_BY_NAME(MovieSceneChaosCacheTemplate)
+
 DECLARE_CYCLE_STAT(TEXT("Chaos Cache Evaluate"), MovieSceneEval_ChaosCache_Evaluate, STATGROUP_MovieSceneEval);
 DECLARE_CYCLE_STAT(TEXT("Chaos Cache Token Execute"), MovieSceneEval_ChaosCache_TokenExecute, STATGROUP_MovieSceneEval);
 
@@ -111,4 +113,4 @@ void FMovieSceneChaosCacheSectionTemplate::Evaluate(const FMovieSceneEvaluationO
 {
 	MOVIESCENE_DETAILED_SCOPE_CYCLE_COUNTER(MovieSceneEval_ChaosCache_Evaluate)
 	ExecutionTokens.Add(FChaosCacheExecutionToken(Params));
-}
\ No newline at end of file
+}
diff --git a/Engine/Plugins/Experimental/ChaosCaching/Source/ChaosCaching/Private/Chaos/Sequencer/MovieSceneChaosCacheTrack.cpp b/Engine/Plugins/Experimental/ChaosCaching/Source/ChaosCaching/Private/Chaos/Sequencer/MovieSceneChaosCacheTrack.cpp
--- a/Engine/Plugins/Experimental/ChaosCaching/Source/ChaosCaching/Private/Chaos/Sequencer/MovieSceneChaosCacheTrack.cpp
+++ b/Engine/Plugins/Experimental/ChaosCaching/Source/ChaosCaching/Private/Chaos/Sequencer/MovieSceneChaosCacheTrack.cpp
@@ -7,6 +7,8 @@
 #include "Chaos/Sequencer/MovieSceneChaosCacheSection.h"
 #include "MovieSceneChaosCacheTemplate.h"
 
+#include UE_INLINE_GENERATED_CPP_BY_NAME(MovieSceneChaosCacheTrack)
+
 #define LOCTEXT_NAMESPACE "MovieSceneChaosCacheTrack"
 
 /* UMovieSceneChaosCacheTrack constructor
diff --git a/Engine/Plugins/Experimental/Chooser/Source/Chooser/Private/Chooser.cpp b/Engine/Plugins/Experimental/Chooser/Source/Chooser/Private/Chooser.cpp
--- a/Engine/Plugins/Experimental/Chooser/Source/Chooser/Private/Chooser.cpp
+++ b/Engine/Plugins/Experimental/Chooser/Source/Chooser/Private/Chooser.cpp
@@ -3,6 +3,8 @@
 #include "ChooserFunctionLibrary.h"
 #include "ChooserPropertyAccess.h"
 
+#include UE_INLINE_GENERATED_CPP_BY_NAME(Chooser)
+
 UChooserTable::UChooserTable(const FObjectInitializer& Initializer)
 	:Super(Initializer)
 {
diff --git a/Engine/Plugins/Experimental/DataInterfaceGraph/Source/DataInterfaceGraph/Private/DataInterfaceGraph.cpp b/Engine/Plugins/Experimental/DataInterfaceGraph/Source/DataInterfaceGraph/Private/DataInterfaceGraph.cpp
--- a/Engine/Plugins/Experimental/DataInterfaceGraph/Source/DataInterfaceGraph/Private/DataInterfaceGraph.cpp
+++ b/Engine/Plugins/Experimental/DataInterfaceGraph/Source/DataInterfaceGraph/Private/DataInterfaceGraph.cpp
@@ -5,6 +5,8 @@
 #include "RigUnit_DataInterfaceBeginExecution.h"
 #include "DataInterfaceUnitContext.h"
 
+#include UE_INLINE_GENERATED_CPP_BY_NAME(DataInterfaceGraph)
+
 namespace UE::DataInterfaceGraph
 {
 const FName EntryPointName("GetData");
@@ -43,4 +45,4 @@ bool UDataInterfaceGraph::GetDataImpl(const UE::DataInterface::FContext& Context
 TArray<FRigVMExternalVariable> UDataInterfaceGraph::GetRigVMExternalVariables()
 {
 	return TArray<FRigVMExternalVariable>(); 
-}
\ No newline at end of file
+}
diff --git a/Engine/Plugins/Experimental/DataInterfaceGraph/Source/DataInterfaceGraph/Private/DataInterface_Bool.cpp b/Engine/Plugins/Experimental/DataInterfaceGraph/Source/DataInterfaceGraph/Private/DataInterface_Bool.cpp
--- a/Engine/Plugins/Experimental/DataInterfaceGraph/Source/DataInterfaceGraph/Private/DataInterface_Bool.cpp
+++ b/Engine/Plugins/Experimental/DataInterfaceGraph/Source/DataInterfaceGraph/Private/DataInterface_Bool.cpp
@@ -5,6 +5,8 @@
 #include "DataInterfaceTypes.h"
 #include "DataInterfaceKernel.h"
 
+#include UE_INLINE_GENERATED_CPP_BY_NAME(DataInterface_Bool)
+
 bool UDataInterface_Bool_And::GetDataImpl(const UE::DataInterface::FContext& Context) const
 {
 	using namespace UE::DataInterface;
diff --git a/Engine/Plugins/Experimental/DataInterfaceGraph/Source/DataInterfaceGraph/Private/DataInterface_Float.cpp b/Engine/Plugins/Experimental/DataInterfaceGraph/Source/DataInterfaceGraph/Private/DataInterface_Float.cpp
--- a/Engine/Plugins/Experimental/DataInterfaceGraph/Source/DataInterfaceGraph/Private/DataInterface_Float.cpp
+++ b/Engine/Plugins/Experimental/DataInterfaceGraph/Source/DataInterfaceGraph/Private/DataInterface_Float.cpp
@@ -5,6 +5,8 @@
 #include "DataInterfaceTypes.h"
 #include "DataInterfaceKernel.h"
 
+#include UE_INLINE_GENERATED_CPP_BY_NAME(DataInterface_Float)
+
 #define LOCTEXT_NAMESPACE "DataInterface"
 
 bool UDataInterface_Float_Multiply::GetDataImpl(const UE::DataInterface::FContext& Context) const
@@ -107,4 +109,4 @@ bool UDataInterface_Float_SpringInterp::GetDataImpl(const UE::DataInterface::FCo
 }
 
 
-#undef LOCTEXT_NAMESPACE
\ No newline at end of file
+#undef LOCTEXT_NAMESPACE
diff --git a/Engine/Plugins/Experimental/DataInterfaceGraph/Source/DataInterfaceGraph/Private/DataInterface_Name.cpp b/Engine/Plugins/Experimental/DataInterfaceGraph/Source/DataInterfaceGraph/Private/DataInterface_Name.cpp
--- a/Engine/Plugins/Experimental/DataInterfaceGraph/Source/DataInterfaceGraph/Private/DataInterface_Name.cpp
+++ b/Engine/Plugins/Experimental/DataInterfaceGraph/Source/DataInterfaceGraph/Private/DataInterface_Name.cpp
@@ -2,3 +2,5 @@
 
 #include "DataInterface_Name.h"
 
+#include UE_INLINE_GENERATED_CPP_BY_NAME(DataInterface_Name)
+
diff --git a/Engine/Plugins/Experimental/DataInterfaceGraph/Source/DataInterfaceGraph/Private/DataInterface_Wrapper.cpp b/Engine/Plugins/Experimental/DataInterfaceGraph/Source/DataInterfaceGraph/Private/DataInterface_Wrapper.cpp
--- a/Engine/Plugins/Experimental/DataInterfaceGraph/Source/DataInterfaceGraph/Private/DataInterface_Wrapper.cpp
+++ b/Engine/Plugins/Experimental/DataInterfaceGraph/Source/DataInterfaceGraph/Private/DataInterface_Wrapper.cpp
@@ -2,6 +2,8 @@
 
 #include "DataInterface_Wrapper.h"
 
+#include UE_INLINE_GENERATED_CPP_BY_NAME(DataInterface_Wrapper)
+
 FName UDataInterface_Wrapper::GetReturnTypeNameImpl() const
 {
 	// We mimic the return type of our output
diff --git a/Engine/Plugins/Experimental/DataInterfaceGraph/Source/DataInterfaceGraph/Private/RigUnit_DataInterfaceBeginExecution.cpp b/Engine/Plugins/Experimental/DataInterfaceGraph/Source/DataInterfaceGraph/Private/RigUnit_DataInterfaceBeginExecution.cpp
--- a/Engine/Plugins/Experimental/DataInterfaceGraph/Source/DataInterfaceGraph/Private/RigUnit_DataInterfaceBeginExecution.cpp
+++ b/Engine/Plugins/Experimental/DataInterfaceGraph/Source/DataInterfaceGraph/Private/RigUnit_DataInterfaceBeginExecution.cpp
@@ -3,6 +3,8 @@
 #include "RigUnit_DataInterfaceBeginExecution.h"
 #include "Units/RigUnitContext.h"
 
+#include UE_INLINE_GENERATED_CPP_BY_NAME(RigUnit_DataInterfaceBeginExecution)
+
 FName FRigUnit_DataInterfaceBeginExecution::EventName = TEXT("Update");
 
 FRigUnit_DataInterfaceBeginExecution_Execute()
diff --git a/Engine/Plugins/Experimental/DataInterfaceGraph/Source/DataInterfaceGraph/Private/RigUnit_DataInterfaceEndExecution.cpp b/Engine/Plugins/Experimental/DataInterfaceGraph/Source/DataInterfaceGraph/Private/RigUnit_DataInterfaceEndExecution.cpp
--- a/Engine/Plugins/Experimental/DataInterfaceGraph/Source/DataInterfaceGraph/Private/RigUnit_DataInterfaceEndExecution.cpp
+++ b/Engine/Plugins/Experimental/DataInterfaceGraph/Source/DataInterfaceGraph/Private/RigUnit_DataInterfaceEndExecution.cpp
@@ -3,6 +3,8 @@
 #include "RigUnit_DataInterfaceEndExecution.h"
 #include "Units/RigUnitContext.h"
 
+#include UE_INLINE_GENERATED_CPP_BY_NAME(RigUnit_DataInterfaceEndExecution)
+
 FRigUnit_DataInterfaceEndExecution_Bool_Execute()
 {
 	if(Context.State == EControlRigState::Update)
diff --git a/Engine/Plugins/Experimental/DataInterfaceGraph/Source/DataInterfaceGraph/Private/RigUnit_DataInterfaceParameter.cpp b/Engine/Plugins/Experimental/DataInterfaceGraph/Source/DataInterfaceGraph/Private/RigUnit_DataInterfaceParameter.cpp
--- a/Engine/Plugins/Experimental/DataInterfaceGraph/Source/DataInterfaceGraph/Private/RigUnit_DataInterfaceParameter.cpp
+++ b/Engine/Plugins/Experimental/DataInterfaceGraph/Source/DataInterfaceGraph/Private/RigUnit_DataInterfaceParameter.cpp
@@ -7,6 +7,8 @@
 #include "DataInterfaceTypes.h"
 #include "DataInterfaceKernel.h"
 
+#include UE_INLINE_GENERATED_CPP_BY_NAME(RigUnit_DataInterfaceParameter)
+
 bool FRigUnit_DataInterfaceParameter::GetParameterInternal(FName InName, const FDataInterfaceUnitContext& InContext, void* OutResult)
 {
 	return true;
diff --git a/Engine/Plugins/Experimental/GeometryScripting/Source/GeometryScriptingCore/Private/GeometryScriptSelectionTypes.cpp b/Engine/Plugins/Experimental/GeometryScripting/Source/GeometryScriptingCore/Private/GeometryScriptSelectionTypes.cpp
--- a/Engine/Plugins/Experimental/GeometryScripting/Source/GeometryScriptingCore/Private/GeometryScriptSelectionTypes.cpp
+++ b/Engine/Plugins/Experimental/GeometryScripting/Source/GeometryScriptingCore/Private/GeometryScriptSelectionTypes.cpp
@@ -6,6 +6,8 @@
 #include "UDynamicMesh.h"
 #include "Misc/StringBuilder.h"
 
+#include UE_INLINE_GENERATED_CPP_BY_NAME(GeometryScriptSelectionTypes)
+
 using namespace UE::Geometry;
 
 
@@ -476,4 +478,4 @@ void FGeometryScriptMeshSelection::ProcessByVertexID(const UE::Geometry::FDynami
 			}
 		}
 	}
-}
\ No newline at end of file
+}
diff --git a/Engine/Plugins/Experimental/GeometryScripting/Source/GeometryScriptingCore/Private/MeshBoneWeightFunctions.cpp b/Engine/Plugins/Experimental/GeometryScripting/Source/GeometryScriptingCore/Private/MeshBoneWeightFunctions.cpp
--- a/Engine/Plugins/Experimental/GeometryScripting/Source/GeometryScriptingCore/Private/MeshBoneWeightFunctions.cpp
+++ b/Engine/Plugins/Experimental/GeometryScripting/Source/GeometryScriptingCore/Private/MeshBoneWeightFunctions.cpp
@@ -8,6 +8,8 @@
 #include "UDynamicMesh.h"
 #include "BoneWeights.h"
 
+#include UE_INLINE_GENERATED_CPP_BY_NAME(MeshBoneWeightFunctions)
+
 using namespace UE::Geometry;
 using namespace UE::AnimationCore;
 
@@ -227,4 +229,4 @@ UDynamicMesh* UGeometryScriptLibrary_MeshBoneWeightFunctions::SetVertexBoneWeigh
 
 
 
-#undef LOCTEXT_NAMESPACE
\ No newline at end of file
+#undef LOCTEXT_NAMESPACE
diff --git a/Engine/Plugins/Experimental/GeometryScripting/Source/GeometryScriptingCore/Private/MeshSelectionFunctions.cpp b/Engine/Plugins/Experimental/GeometryScripting/Source/GeometryScriptingCore/Private/MeshSelectionFunctions.cpp
--- a/Engine/Plugins/Experimental/GeometryScripting/Source/GeometryScriptingCore/Private/MeshSelectionFunctions.cpp
+++ b/Engine/Plugins/Experimental/GeometryScripting/Source/GeometryScriptingCore/Private/MeshSelectionFunctions.cpp
@@ -13,6 +13,8 @@
 #include "Selections/MeshFaceSelection.h"
 #include "Selections/MeshVertexSelection.h"
 
+#include UE_INLINE_GENERATED_CPP_BY_NAME(MeshSelectionFunctions)
+
 using namespace UE::Geometry;
 
 #define LOCTEXT_NAMESPACE "UGeometryScriptLibrary_MeshSelectionFunctions"
@@ -963,4 +965,4 @@ UDynamicMesh* UGeometryScriptLibrary_MeshSelectionFunctions::ExpandContractMeshS
 
 
 
-#undef LOCTEXT_NAMESPACE
\ No newline at end of file
+#undef LOCTEXT_NAMESPACE
diff --git a/Engine/Plugins/Experimental/GeometryScripting/Source/GeometryScriptingCore/Private/MeshSelectionQueryFunctions.cpp b/Engine/Plugins/Experimental/GeometryScripting/Source/GeometryScriptingCore/Private/MeshSelectionQueryFunctions.cpp
--- a/Engine/Plugins/Experimental/GeometryScripting/Source/GeometryScriptingCore/Private/MeshSelectionQueryFunctions.cpp
+++ b/Engine/Plugins/Experimental/GeometryScripting/Source/GeometryScriptingCore/Private/MeshSelectionQueryFunctions.cpp
@@ -10,6 +10,8 @@
 #include "Selections/MeshConnectedComponents.h"
 #include "MeshRegionBoundaryLoops.h"
 
+#include UE_INLINE_GENERATED_CPP_BY_NAME(MeshSelectionQueryFunctions)
+
 using namespace UE::Geometry;
 
 #define LOCTEXT_NAMESPACE "UGeometryScriptLibrary_MeshSelectionQueryFunctions"
@@ -122,4 +124,4 @@ UDynamicMesh* UGeometryScriptLibrary_MeshSelectionQueryFunctions::GetMeshSelecti
 }
 
 
-#undef LOCTEXT_NAMESPACE
\ No newline at end of file
+#undef LOCTEXT_NAMESPACE
diff --git a/Engine/Plugins/Experimental/GeometryScripting/Source/GeometryScriptingCore/Private/ShapeFunctions.cpp b/Engine/Plugins/Experimental/GeometryScripting/Source/GeometryScriptingCore/Private/ShapeFunctions.cpp
--- a/Engine/Plugins/Experimental/GeometryScripting/Source/GeometryScriptingCore/Private/ShapeFunctions.cpp
+++ b/Engine/Plugins/Experimental/GeometryScripting/Source/GeometryScriptingCore/Private/ShapeFunctions.cpp
@@ -3,6 +3,8 @@
 #include "GeometryScript/ShapeFunctions.h"
 #include "VectorTypes.h"
 
+#include UE_INLINE_GENERATED_CPP_BY_NAME(ShapeFunctions)
+
 using namespace UE::Geometry;
 
 #define LOCTEXT_NAMESPACE "UGeometryScriptLibrary_BoxFunctions"
@@ -125,4 +127,4 @@ bool UGeometryScriptLibrary_BoxFunctions::TestBoxSphereIntersection(const FBox&
 
 
 
-#undef LOCTEXT_NAMESPACE
\ No newline at end of file
+#undef LOCTEXT_NAMESPACE
diff --git a/Engine/Plugins/Experimental/Mutable/Source/CustomizableObject/Private/MuCO/CustomizableInstanceLODManagement.cpp b/Engine/Plugins/Experimental/Mutable/Source/CustomizableObject/Private/MuCO/CustomizableInstanceLODManagement.cpp
--- a/Engine/Plugins/Experimental/Mutable/Source/CustomizableObject/Private/MuCO/CustomizableInstanceLODManagement.cpp
+++ b/Engine/Plugins/Experimental/Mutable/Source/CustomizableObject/Private/MuCO/CustomizableInstanceLODManagement.cpp
@@ -31,6 +31,8 @@
 #include "UObject/ObjectPtr.h"
 #include "UObject/UObjectIterator.h"
 
+#include UE_INLINE_GENERATED_CPP_BY_NAME(CustomizableInstanceLODManagement)
+
 #if WITH_EDITOR
 #include "Editor.h"
 #include "Editor/EditorEngine.h"
diff --git a/Engine/Plugins/Experimental/Mutable/Source/CustomizableObject/Private/MuCO/CustomizableObject.cpp b/Engine/Plugins/Experimental/Mutable/Source/CustomizableObject/Private/MuCO/CustomizableObject.cpp
--- a/Engine/Plugins/Experimental/Mutable/Source/CustomizableObject/Private/MuCO/CustomizableObject.cpp
+++ b/Engine/Plugins/Experimental/Mutable/Source/CustomizableObject/Private/MuCO/CustomizableObject.cpp
@@ -49,6 +49,8 @@
 #include "UObject/UObjectMarks.h"
 #include "UObject/UnrealNames.h"
 
+#include UE_INLINE_GENERATED_CPP_BY_NAME(CustomizableObject)
+
 class UMaterialInterface;
 class UPhysicsAsset;
 class USkeletalMesh;
diff --git a/Engine/Plugins/Experimental/Mutable/Source/CustomizableObject/Private/MuCO/CustomizableObjectDGGUI.cpp b/Engine/Plugins/Experimental/Mutable/Source/CustomizableObject/Private/MuCO/CustomizableObjectDGGUI.cpp
--- a/Engine/Plugins/Experimental/Mutable/Source/CustomizableObject/Private/MuCO/CustomizableObjectDGGUI.cpp
+++ b/Engine/Plugins/Experimental/Mutable/Source/CustomizableObject/Private/MuCO/CustomizableObjectDGGUI.cpp
@@ -24,6 +24,8 @@
 #include "UObject/SoftObjectPtr.h"
 #include "UObject/UObjectIterator.h"
 
+#include UE_INLINE_GENERATED_CPP_BY_NAME(CustomizableObjectDGGUI)
+
 void UDGGUI::OpenDGGUI(const int32 SlotID, UCustomizableSkeletalComponent* SelectedCustomizableSkeletalComponent, const UWorld* CurrentWorld, const int32 PlayerIndex)
 {
 #if WITH_EDITOR
@@ -155,4 +157,4 @@ void UDGGUI::CustomizableSkeletalMeshUpdated()
 		//CustomizableComponent->UpdatedDelegate.Unbind();
 	}
 #endif // WITH_EDITOR
-}
\ No newline at end of file
+}
diff --git a/Engine/Plugins/Experimental/Mutable/Source/CustomizableObject/Private/MuCO/CustomizableObjectIdentifier.cpp b/Engine/Plugins/Experimental/Mutable/Source/CustomizableObject/Private/MuCO/CustomizableObjectIdentifier.cpp
--- a/Engine/Plugins/Experimental/Mutable/Source/CustomizableObject/Private/MuCO/CustomizableObjectIdentifier.cpp
+++ b/Engine/Plugins/Experimental/Mutable/Source/CustomizableObject/Private/MuCO/CustomizableObjectIdentifier.cpp
@@ -2,6 +2,8 @@
 
 #include "MuCO/CustomizableObjectIdentifier.h"
 
+#include UE_INLINE_GENERATED_CPP_BY_NAME(CustomizableObjectIdentifier)
+
 
 FCustomizableObjectIdPair::FCustomizableObjectIdPair(FString ObjectGroupName, FString ObjectName)
 	: CustomizableObjectGroupName(ObjectGroupName)
diff --git a/Engine/Plugins/Experimental/Mutable/Source/CustomizableObject/Private/MuCO/CustomizableObjectInstance.cpp b/Engine/Plugins/Experimental/Mutable/Source/CustomizableObject/Private/MuCO/CustomizableObjectInstance.cpp
--- a/Engine/Plugins/Experimental/Mutable/Source/CustomizableObject/Private/MuCO/CustomizableObjectInstance.cpp
+++ b/Engine/Plugins/Experimental/Mutable/Source/CustomizableObject/Private/MuCO/CustomizableObjectInstance.cpp
@@ -129,6 +129,8 @@
 #include "MuR/System.h"
 #include "PhysicsEngine/PhysicsConstraintTemplate.h"
 
+#include UE_INLINE_GENERATED_CPP_BY_NAME(CustomizableObjectInstance)
+
 UTexture2D* UCustomizableInstancePrivateData::CreateTexture()
 {
 	UTexture2D* NewTexture = NewObject<UTexture2D>(
diff --git a/Engine/Plugins/Experimental/Mutable/Source/CustomizableObject/Private/MuCO/CustomizableObjectInstanceDescriptor.cpp b/Engine/Plugins/Experimental/Mutable/Source/CustomizableObject/Private/MuCO/CustomizableObjectInstanceDescriptor.cpp
--- a/Engine/Plugins/Experimental/Mutable/Source/CustomizableObject/Private/MuCO/CustomizableObjectInstanceDescriptor.cpp
+++ b/Engine/Plugins/Experimental/Mutable/Source/CustomizableObject/Private/MuCO/CustomizableObjectInstanceDescriptor.cpp
@@ -19,6 +19,8 @@
 #include "Templates/TypeHash.h"
 #include "Trace/Detail/Channel.h"
 
+#include UE_INLINE_GENERATED_CPP_BY_NAME(CustomizableObjectInstanceDescriptor)
+
 
 const FString EmptyString;
 
diff --git a/Engine/Plugins/Experimental/Mutable/Source/CustomizableObject/Private/MuCO/CustomizableSkeletalComponent.cpp b/Engine/Plugins/Experimental/Mutable/Source/CustomizableObject/Private/MuCO/CustomizableSkeletalComponent.cpp
--- a/Engine/Plugins/Experimental/Mutable/Source/CustomizableObject/Private/MuCO/CustomizableSkeletalComponent.cpp
+++ b/Engine/Plugins/Experimental/Mutable/Source/CustomizableObject/Private/MuCO/CustomizableSkeletalComponent.cpp
@@ -32,6 +32,8 @@
 #include "UObject/WeakObjectPtr.h"
 #include "UObject/WeakObjectPtrTemplates.h"
 
+#include UE_INLINE_GENERATED_CPP_BY_NAME(CustomizableSkeletalComponent)
+
 
 UCustomizableSkeletalComponent::UCustomizableSkeletalComponent()
 {
diff --git a/Engine/Plugins/Experimental/Mutable/Source/CustomizableObject/Private/MuCO/CustomizableSkeletalMeshActor.cpp b/Engine/Plugins/Experimental/Mutable/Source/CustomizableObject/Private/MuCO/CustomizableSkeletalMeshActor.cpp
--- a/Engine/Plugins/Experimental/Mutable/Source/CustomizableObject/Private/MuCO/CustomizableSkeletalMeshActor.cpp
+++ b/Engine/Plugins/Experimental/Mutable/Source/CustomizableObject/Private/MuCO/CustomizableSkeletalMeshActor.cpp
@@ -9,6 +9,8 @@
 #include "MuCO/CustomizableSkeletalComponent.h"
 #include "UObject/NameTypes.h"
 
+#include UE_INLINE_GENERATED_CPP_BY_NAME(CustomizableSkeletalMeshActor)
+
 #define LOCTEXT_NAMESPACE "CustomizableObject"
 
 
diff --git a/Engine/Plugins/Experimental/Mutable/Source/CustomizableObject/Private/MuCO/DefaultImageProvider.cpp b/Engine/Plugins/Experimental/Mutable/Source/CustomizableObject/Private/MuCO/DefaultImageProvider.cpp
--- a/Engine/Plugins/Experimental/Mutable/Source/CustomizableObject/Private/MuCO/DefaultImageProvider.cpp
+++ b/Engine/Plugins/Experimental/Mutable/Source/CustomizableObject/Private/MuCO/DefaultImageProvider.cpp
@@ -5,6 +5,8 @@
 #include "MuCO/CustomizableObjectInstance.h"
 #include "MuR/System.h"
 
+#include UE_INLINE_GENERATED_CPP_BY_NAME(DefaultImageProvider)
+
 
 namespace UDefaultImageProviderCVars
 {
diff --git a/Engine/Plugins/Experimental/Mutable/Source/CustomizableObject/Private/MuCO/MultilayerProjector.cpp b/Engine/Plugins/Experimental/Mutable/Source/CustomizableObject/Private/MuCO/MultilayerProjector.cpp
--- a/Engine/Plugins/Experimental/Mutable/Source/CustomizableObject/Private/MuCO/MultilayerProjector.cpp
+++ b/Engine/Plugins/Experimental/Mutable/Source/CustomizableObject/Private/MuCO/MultilayerProjector.cpp
@@ -10,6 +10,8 @@
 #include "MuCO/CustomizableObjectParameterTypeDefinitions.h"
 #include "Templates/Tuple.h"
 
+#include UE_INLINE_GENERATED_CPP_BY_NAME(MultilayerProjector)
+
 
 void FMultilayerProjectorLayer::Read(const FCustomizableObjectInstanceDescriptor& Descriptor, const FMultilayerProjector& MultilayerProjector, const int32 Index)
 {
@@ -530,4 +532,4 @@ uint32 GetTypeHash(const FMultilayerProjector& Key)
 	}
 	
 	return Hash;
-}
\ No newline at end of file
+}
diff --git a/Engine/Plugins/Experimental/Mutable/Source/CustomizableObjectPopulation/Private/MuCOP/CustomizableObjectPopulation.cpp b/Engine/Plugins/Experimental/Mutable/Source/CustomizableObjectPopulation/Private/MuCOP/CustomizableObjectPopulation.cpp
--- a/Engine/Plugins/Experimental/Mutable/Source/CustomizableObjectPopulation/Private/MuCOP/CustomizableObjectPopulation.cpp
+++ b/Engine/Plugins/Experimental/Mutable/Source/CustomizableObjectPopulation/Private/MuCOP/CustomizableObjectPopulation.cpp
@@ -19,6 +19,8 @@
 #include "Templates/SharedPointer.h"
 #include "Templates/Tuple.h"
 
+#include UE_INLINE_GENERATED_CPP_BY_NAME(CustomizableObjectPopulation)
+
 class ITargetPlatform;
 class UCustomizableObjectInstance;
 
diff --git a/Engine/Plugins/Experimental/Mutable/Source/CustomizableObjectPopulation/Private/MuCOP/CustomizableObjectPopulationCharacteristic.cpp b/Engine/Plugins/Experimental/Mutable/Source/CustomizableObjectPopulation/Private/MuCOP/CustomizableObjectPopulationCharacteristic.cpp
--- a/Engine/Plugins/Experimental/Mutable/Source/CustomizableObjectPopulation/Private/MuCOP/CustomizableObjectPopulationCharacteristic.cpp
+++ b/Engine/Plugins/Experimental/Mutable/Source/CustomizableObjectPopulation/Private/MuCOP/CustomizableObjectPopulationCharacteristic.cpp
@@ -2,6 +2,8 @@
 
 #include "MuCOP/CustomizableObjectPopulationCharacteristic.h"
 
+#include UE_INLINE_GENERATED_CPP_BY_NAME(CustomizableObjectPopulationCharacteristic)
+
 #define LOCTEXT_NAMESPACE "CustomizableObjectPopulationCharacteristic"
 
 
diff --git a/Engine/Plugins/Experimental/Mutable/Source/CustomizableObjectPopulation/Private/MuCOP/CustomizableObjectPopulationClass.cpp b/Engine/Plugins/Experimental/Mutable/Source/CustomizableObjectPopulation/Private/MuCOP/CustomizableObjectPopulationClass.cpp
--- a/Engine/Plugins/Experimental/Mutable/Source/CustomizableObjectPopulation/Private/MuCOP/CustomizableObjectPopulationClass.cpp
+++ b/Engine/Plugins/Experimental/Mutable/Source/CustomizableObjectPopulation/Private/MuCOP/CustomizableObjectPopulationClass.cpp
@@ -6,6 +6,8 @@
 #include "MuCOP/CustomizableObjectPopulationCustomVersion.h"
 #include "Serialization/Archive.h"
 
+#include UE_INLINE_GENERATED_CPP_BY_NAME(CustomizableObjectPopulationClass)
+
 #define LOCTEXT_NAMESPACE "CustomizableObjectPopulationClass"
 
 using namespace CustomizableObjectPopulation;
@@ -17,4 +19,4 @@ void UCustomizableObjectPopulationClass::Serialize(FArchive& Ar)
 	Ar.UsingCustomVersion(FCustomizableObjectPopulationCustomVersion::GUID);
 }
 
-#undef LOCTEXT_NAMESPACE
\ No newline at end of file
+#undef LOCTEXT_NAMESPACE
diff --git a/Engine/Plugins/Experimental/Mutable/Source/CustomizableObjectPopulation/Private/MuCOP/CustomizableObjectPopulationConstraint.cpp b/Engine/Plugins/Experimental/Mutable/Source/CustomizableObjectPopulation/Private/MuCOP/CustomizableObjectPopulationConstraint.cpp
--- a/Engine/Plugins/Experimental/Mutable/Source/CustomizableObjectPopulation/Private/MuCOP/CustomizableObjectPopulationConstraint.cpp
+++ b/Engine/Plugins/Experimental/Mutable/Source/CustomizableObjectPopulation/Private/MuCOP/CustomizableObjectPopulationConstraint.cpp
@@ -2,8 +2,11 @@
 
 #include "MuCOP/CustomizableObjectPopulationConstraint.h"
 
+#include UE_INLINE_GENERATED_CPP_BY_NAME(CustomizableObjectPopulationConstraint)
+
 #define LOCTEXT_NAMESPACE "CustomizableObjectPopulationConstraint"
 
 
 
 #undef LOCTEXT_NAMESPACE
+
diff --git a/Engine/Plugins/Experimental/Mutable/Source/CustomizableObjectPopulation/Private/MuCOP/CustomizableObjectPopulationGenerator.cpp b/Engine/Plugins/Experimental/Mutable/Source/CustomizableObjectPopulation/Private/MuCOP/CustomizableObjectPopulationGenerator.cpp
--- a/Engine/Plugins/Experimental/Mutable/Source/CustomizableObjectPopulation/Private/MuCOP/CustomizableObjectPopulationGenerator.cpp
+++ b/Engine/Plugins/Experimental/Mutable/Source/CustomizableObjectPopulation/Private/MuCOP/CustomizableObjectPopulationGenerator.cpp
@@ -13,6 +13,8 @@
 #include "MuCOP/CustomizableObjectPopulationSamplers.h"
 #include "Serialization/Archive.h"
 
+#include UE_INLINE_GENERATED_CPP_BY_NAME(CustomizableObjectPopulationGenerator)
+
 class FString;
 
 void FCustomizableObjectPopulationClassGenerator::GenerateParameters(
diff --git a/Engine/Plugins/Experimental/NNI/Source/NeuralNetworkInference/Private/NeuralEnumClasses.cpp b/Engine/Plugins/Experimental/NNI/Source/NeuralNetworkInference/Private/NeuralEnumClasses.cpp
--- a/Engine/Plugins/Experimental/NNI/Source/NeuralNetworkInference/Private/NeuralEnumClasses.cpp
+++ b/Engine/Plugins/Experimental/NNI/Source/NeuralNetworkInference/Private/NeuralEnumClasses.cpp
@@ -3,6 +3,8 @@
 #include "NeuralEnumClasses.h"
 #include "NeuralNetworkInferenceUtils.h"
 
+#include UE_INLINE_GENERATED_CPP_BY_NAME(NeuralEnumClasses)
+
 
 
 /* FNeuralNetworkInferenceUtils public functions
diff --git a/Engine/Plugins/Experimental/NNI/Source/NeuralNetworkInference/Private/NeuralNetwork.cpp b/Engine/Plugins/Experimental/NNI/Source/NeuralNetworkInference/Private/NeuralNetwork.cpp
--- a/Engine/Plugins/Experimental/NNI/Source/NeuralNetworkInference/Private/NeuralNetwork.cpp
+++ b/Engine/Plugins/Experimental/NNI/Source/NeuralNetworkInference/Private/NeuralNetwork.cpp
@@ -13,6 +13,8 @@
 #include "RenderingThread.h"
 #include "RHI.h"
 
+#include UE_INLINE_GENERATED_CPP_BY_NAME(NeuralNetwork)
+
 
 
 /* FPrivateNeuralNetwork functions
diff --git a/Engine/Plugins/Experimental/NNI/Source/NeuralNetworkInference/Private/NeuralTensor.cpp b/Engine/Plugins/Experimental/NNI/Source/NeuralNetworkInference/Private/NeuralTensor.cpp
--- a/Engine/Plugins/Experimental/NNI/Source/NeuralNetworkInference/Private/NeuralTensor.cpp
+++ b/Engine/Plugins/Experimental/NNI/Source/NeuralNetworkInference/Private/NeuralTensor.cpp
@@ -10,6 +10,8 @@
 #include "RenderGraphUtils.h"
 #include "RenderingThread.h"
 
+#include UE_INLINE_GENERATED_CPP_BY_NAME(NeuralTensor)
+
 
 
 /* FPrivateNeuralTensor auxiliary functions
diff --git a/Engine/Plugins/Experimental/NNI/Source/Private/ModelProto/Private/ModelProto.cpp b/Engine/Plugins/Experimental/NNI/Source/Private/ModelProto/Private/ModelProto.cpp
--- a/Engine/Plugins/Experimental/NNI/Source/Private/ModelProto/Private/ModelProto.cpp
+++ b/Engine/Plugins/Experimental/NNI/Source/Private/ModelProto/Private/ModelProto.cpp
@@ -4,6 +4,8 @@
 #include "ModelProtoUtils.h"
 #include "ModelProtoStringParser.h"
 
+#include UE_INLINE_GENERATED_CPP_BY_NAME(ModelProto)
+
 
 
 /* FTensorProtoSegment structors
diff --git a/Engine/Plugins/Experimental/NNI/Source/Private/NeuralNetworkInferenceQA/Private/NeuralNetworkInferenceQA.cpp b/Engine/Plugins/Experimental/NNI/Source/Private/NeuralNetworkInferenceQA/Private/NeuralNetworkInferenceQA.cpp
--- a/Engine/Plugins/Experimental/NNI/Source/Private/NeuralNetworkInferenceQA/Private/NeuralNetworkInferenceQA.cpp
+++ b/Engine/Plugins/Experimental/NNI/Source/Private/NeuralNetworkInferenceQA/Private/NeuralNetworkInferenceQA.cpp
@@ -5,6 +5,8 @@
 #include "UnitTester.h"
 #include "Misc/Paths.h"
 
+#include UE_INLINE_GENERATED_CPP_BY_NAME(NeuralNetworkInferenceQA)
+
 
 
 /* UNeuralNetworkInferenceQA static functions
diff --git a/Engine/Plugins/Experimental/NNI/Source/Private/NeuralNetworkInferenceQA/Private/NeuralNetworkInferenceQAAsset.cpp b/Engine/Plugins/Experimental/NNI/Source/Private/NeuralNetworkInferenceQA/Private/NeuralNetworkInferenceQAAsset.cpp
--- a/Engine/Plugins/Experimental/NNI/Source/Private/NeuralNetworkInferenceQA/Private/NeuralNetworkInferenceQAAsset.cpp
+++ b/Engine/Plugins/Experimental/NNI/Source/Private/NeuralNetworkInferenceQA/Private/NeuralNetworkInferenceQAAsset.cpp
@@ -7,6 +7,8 @@
 #include "Misc/Paths.h"
 #include "UObject/SavePackage.h"
 
+#include UE_INLINE_GENERATED_CPP_BY_NAME(NeuralNetworkInferenceQAAsset)
+
 
 /* FNeuralNetworkInferenceQAOperatorTestAsset public functions
  *****************************************************************************/
diff --git a/Engine/Plugins/Experimental/PCG/Source/PCG/Private/Data/PCGDifferenceData.cpp b/Engine/Plugins/Experimental/PCG/Source/PCG/Private/Data/PCGDifferenceData.cpp
--- a/Engine/Plugins/Experimental/PCG/Source/PCG/Private/Data/PCGDifferenceData.cpp
+++ b/Engine/Plugins/Experimental/PCG/Source/PCG/Private/Data/PCGDifferenceData.cpp
@@ -7,6 +7,8 @@
 #include "PCGHelpers.h"
 #include "Metadata/PCGMetadataAccessor.h"
 
+#include UE_INLINE_GENERATED_CPP_BY_NAME(PCGDifferenceData)
+
 namespace PCGDifferenceDataUtils
 {
 	EPCGUnionDensityFunction ToUnionDensityFunction(EPCGDifferenceDensityFunction InDensityFunction)
@@ -226,4 +228,4 @@ UPCGSpatialData* UPCGDifferenceData::CopyInternal() const
 	}
 
 	return NewDifferenceData;
-}
\ No newline at end of file
+}
diff --git a/Engine/Plugins/Experimental/PCG/Source/PCG/Private/Data/PCGIntersectionData.cpp b/Engine/Plugins/Experimental/PCG/Source/PCG/Private/Data/PCGIntersectionData.cpp
--- a/Engine/Plugins/Experimental/PCG/Source/PCG/Private/Data/PCGIntersectionData.cpp
+++ b/Engine/Plugins/Experimental/PCG/Source/PCG/Private/Data/PCGIntersectionData.cpp
@@ -6,6 +6,8 @@
 #include "PCGHelpers.h"
 #include "Metadata/PCGMetadataAccessor.h"
 
+#include UE_INLINE_GENERATED_CPP_BY_NAME(PCGIntersectionData)
+
 namespace PCGIntersectionDataMaths
 {
 	float ComputeDensity(float InDensityA, float InDensityB, EPCGIntersectionDensityFunction InDensityFunction)
@@ -208,4 +210,4 @@ UPCGSpatialData* UPCGIntersectionData::CopyInternal() const
 	NewIntersectionData->CachedStrictBounds = CachedStrictBounds;
 
 	return NewIntersectionData;
-}
\ No newline at end of file
+}
diff --git a/Engine/Plugins/Experimental/PCG/Source/PCG/Private/Data/PCGLandscapeData.cpp b/Engine/Plugins/Experimental/PCG/Source/PCG/Private/Data/PCGLandscapeData.cpp
--- a/Engine/Plugins/Experimental/PCG/Source/PCG/Private/Data/PCGLandscapeData.cpp
+++ b/Engine/Plugins/Experimental/PCG/Source/PCG/Private/Data/PCGLandscapeData.cpp
@@ -13,6 +13,8 @@
 #include "Landscape.h"
 #include "LandscapeEdit.h"
 
+#include UE_INLINE_GENERATED_CPP_BY_NAME(PCGLandscapeData)
+
 void UPCGLandscapeData::Initialize(const TArray<TWeakObjectPtr<ALandscapeProxy>>& InLandscapes, const FBox& InBounds, bool bInHeightOnly, bool bInUseMetadata)
 {
 	for (TWeakObjectPtr<ALandscapeProxy> InLandscape : InLandscapes)
@@ -300,4 +302,4 @@ UPCGSpatialData* UPCGLandscapeData::CopyInternal() const
 	NewLandsacapeData->LandscapeCache = LandscapeCache;
 
 	return NewLandsacapeData;
-}
\ No newline at end of file
+}
diff --git a/Engine/Plugins/Experimental/PCG/Source/PCG/Private/Data/PCGLandscapeSplineData.cpp b/Engine/Plugins/Experimental/PCG/Source/PCG/Private/Data/PCGLandscapeSplineData.cpp
--- a/Engine/Plugins/Experimental/PCG/Source/PCG/Private/Data/PCGLandscapeSplineData.cpp
+++ b/Engine/Plugins/Experimental/PCG/Source/PCG/Private/Data/PCGLandscapeSplineData.cpp
@@ -11,6 +11,8 @@
 #include "LandscapeSplineSegment.h"
 #include "LandscapeSplineControlPoint.h"
 
+#include UE_INLINE_GENERATED_CPP_BY_NAME(PCGLandscapeSplineData)
+
 namespace PCGLandscapeDataHelpers
 {
 	// This function assumes that the A-B segment has a "1" density, while the C-D segment has a "0" density
@@ -226,4 +228,4 @@ UPCGSpatialData* UPCGLandscapeSplineData::CopyInternal() const
 	NewLandscapeSplineData->Spline = Spline;
 
 	return NewLandscapeSplineData;
-}
\ No newline at end of file
+}
diff --git a/Engine/Plugins/Experimental/PCG/Source/PCG/Private/Data/PCGPointData.cpp b/Engine/Plugins/Experimental/PCG/Source/PCG/Private/Data/PCGPointData.cpp
--- a/Engine/Plugins/Experimental/PCG/Source/PCG/Private/Data/PCGPointData.cpp
+++ b/Engine/Plugins/Experimental/PCG/Source/PCG/Private/Data/PCGPointData.cpp
@@ -8,6 +8,8 @@
 #include "GameFramework/Actor.h"
 #include "Misc/ScopeLock.h"
 
+#include UE_INLINE_GENERATED_CPP_BY_NAME(PCGPointData)
+
 namespace PCGPointHelpers
 {
 	bool GetDistanceRatios(const FPCGPoint& InPoint, const FVector& InPosition, FVector& OutRatios)
@@ -468,4 +470,4 @@ UPCGSpatialData* UPCGPointData::CopyInternal() const
 	NewPointData->GetMutablePoints() = GetPoints();
 
 	return NewPointData;
-}
\ No newline at end of file
+}
diff --git a/Engine/Plugins/Experimental/PCG/Source/PCG/Private/Data/PCGPolyLineData.cpp b/Engine/Plugins/Experimental/PCG/Source/PCG/Private/Data/PCGPolyLineData.cpp
--- a/Engine/Plugins/Experimental/PCG/Source/PCG/Private/Data/PCGPolyLineData.cpp
+++ b/Engine/Plugins/Experimental/PCG/Source/PCG/Private/Data/PCGPolyLineData.cpp
@@ -2,6 +2,8 @@
 
 #include "Data/PCGPolyLineData.h"
 
+#include UE_INLINE_GENERATED_CPP_BY_NAME(PCGPolyLineData)
+
 FBox UPCGPolyLineData::GetBounds() const
 {
 	FBox Bounds(EForceInit::ForceInit);
@@ -15,4 +17,4 @@ FBox UPCGPolyLineData::GetBounds() const
 	}
 	
 	return Bounds;
-}
\ No newline at end of file
+}
diff --git a/Engine/Plugins/Experimental/PCG/Source/PCG/Private/Data/PCGPrimitiveData.cpp b/Engine/Plugins/Experimental/PCG/Source/PCG/Private/Data/PCGPrimitiveData.cpp
--- a/Engine/Plugins/Experimental/PCG/Source/PCG/Private/Data/PCGPrimitiveData.cpp
+++ b/Engine/Plugins/Experimental/PCG/Source/PCG/Private/Data/PCGPrimitiveData.cpp
@@ -6,6 +6,8 @@
 #include "Elements/PCGVolumeSampler.h"
 #include "Components/PrimitiveComponent.h"
 
+#include UE_INLINE_GENERATED_CPP_BY_NAME(PCGPrimitiveData)
+
 void UPCGPrimitiveData::Initialize(UPrimitiveComponent* InPrimitive)
 {
 	check(InPrimitive);
@@ -62,4 +64,4 @@ UPCGSpatialData* UPCGPrimitiveData::CopyInternal() const
 	NewPrimitiveData->CachedStrictBounds = CachedStrictBounds;
 
 	return NewPrimitiveData;
-}
\ No newline at end of file
+}
diff --git a/Engine/Plugins/Experimental/PCG/Source/PCG/Private/Data/PCGProjectionData.cpp b/Engine/Plugins/Experimental/PCG/Source/PCG/Private/Data/PCGProjectionData.cpp
--- a/Engine/Plugins/Experimental/PCG/Source/PCG/Private/Data/PCGProjectionData.cpp
+++ b/Engine/Plugins/Experimental/PCG/Source/PCG/Private/Data/PCGProjectionData.cpp
@@ -4,6 +4,8 @@
 #include "Helpers/PCGAsync.h"
 #include "Metadata/PCGMetadataAccessor.h"
 
+#include UE_INLINE_GENERATED_CPP_BY_NAME(PCGProjectionData)
+
 void UPCGProjectionData::Initialize(const UPCGSpatialData* InSource, const UPCGSpatialData* InTarget, const FPCGProjectionParams& InParams)
 {
 	check(InSource && InTarget);
@@ -221,4 +223,4 @@ UPCGSpatialData* UPCGProjectionData::CopyInternal() const
 	CopyBaseProjectionClass(NewProjectionData);
 
 	return NewProjectionData;
-}
\ No newline at end of file
+}
diff --git a/Engine/Plugins/Experimental/PCG/Source/PCG/Private/Data/PCGRenderTargetData.cpp b/Engine/Plugins/Experimental/PCG/Source/PCG/Private/Data/PCGRenderTargetData.cpp
--- a/Engine/Plugins/Experimental/PCG/Source/PCG/Private/Data/PCGRenderTargetData.cpp
+++ b/Engine/Plugins/Experimental/PCG/Source/PCG/Private/Data/PCGRenderTargetData.cpp
@@ -4,6 +4,8 @@
 #include "Kismet/KismetRenderingLibrary.h"
 #include "TextureResource.h"
 
+#include UE_INLINE_GENERATED_CPP_BY_NAME(PCGRenderTargetData)
+
 void UPCGRenderTargetData::Initialize(UTextureRenderTarget2D* InRenderTarget, const FTransform& InTransform)
 {
 	RenderTarget = InRenderTarget;
@@ -49,4 +51,4 @@ UPCGSpatialData* UPCGRenderTargetData::CopyInternal() const
 	NewRenderTargetData->RenderTarget = RenderTarget;
 
 	return NewRenderTargetData;
-}
\ No newline at end of file
+}
diff --git a/Engine/Plugins/Experimental/PCG/Source/PCG/Private/Data/PCGSpatialData.cpp b/Engine/Plugins/Experimental/PCG/Source/PCG/Private/Data/PCGSpatialData.cpp
--- a/Engine/Plugins/Experimental/PCG/Source/PCG/Private/Data/PCGSpatialData.cpp
+++ b/Engine/Plugins/Experimental/PCG/Source/PCG/Private/Data/PCGSpatialData.cpp
@@ -6,6 +6,8 @@
 #include "Data/PCGProjectionData.h"
 #include "Data/PCGUnionData.h"
 
+#include UE_INLINE_GENERATED_CPP_BY_NAME(PCGSpatialData)
+
 UPCGSpatialData::UPCGSpatialData(const FObjectInitializer& ObjectInitializer)
 	: Super(ObjectInitializer)
 {
@@ -160,4 +162,4 @@ UPCGSpatialData* UPCGSpatialData::DuplicateData(const bool bInitializeMetadata)
 	}
 
 	return NewSpatialData;
-}
\ No newline at end of file
+}
diff --git a/Engine/Plugins/Experimental/PCG/Source/PCG/Private/Data/PCGSplineData.cpp b/Engine/Plugins/Experimental/PCG/Source/PCG/Private/Data/PCGSplineData.cpp
--- a/Engine/Plugins/Experimental/PCG/Source/PCG/Private/Data/PCGSplineData.cpp
+++ b/Engine/Plugins/Experimental/PCG/Source/PCG/Private/Data/PCGSplineData.cpp
@@ -9,6 +9,8 @@
 
 #include "Components/SplineComponent.h"
 
+#include UE_INLINE_GENERATED_CPP_BY_NAME(PCGSplineData)
+
 void UPCGSplineData::Initialize(USplineComponent* InSpline)
 {
 	check(InSpline);
@@ -212,4 +214,4 @@ UPCGSpatialData* UPCGSplineProjectionData::CopyInternal() const
 	NewProjectionData->ProjectedPosition = ProjectedPosition;
 
 	return NewProjectionData;
-}
\ No newline at end of file
+}
diff --git a/Engine/Plugins/Experimental/PCG/Source/PCG/Private/Data/PCGSurfaceData.cpp b/Engine/Plugins/Experimental/PCG/Source/PCG/Private/Data/PCGSurfaceData.cpp
--- a/Engine/Plugins/Experimental/PCG/Source/PCG/Private/Data/PCGSurfaceData.cpp
+++ b/Engine/Plugins/Experimental/PCG/Source/PCG/Private/Data/PCGSurfaceData.cpp
@@ -2,7 +2,9 @@
 
 #include "Data/PCGSurfaceData.h"
 
+#include UE_INLINE_GENERATED_CPP_BY_NAME(PCGSurfaceData)
+
 void UPCGSurfaceData::CopyBaseSurfaceData(UPCGSurfaceData* NewSurfaceData) const
 {
 	NewSurfaceData->Transform = Transform;
-}
\ No newline at end of file
+}
diff --git a/Engine/Plugins/Experimental/PCG/Source/PCG/Private/Data/PCGTextureData.cpp b/Engine/Plugins/Experimental/PCG/Source/PCG/Private/Data/PCGTextureData.cpp
--- a/Engine/Plugins/Experimental/PCG/Source/PCG/Private/Data/PCGTextureData.cpp
+++ b/Engine/Plugins/Experimental/PCG/Source/PCG/Private/Data/PCGTextureData.cpp
@@ -7,6 +7,8 @@
 #include "Helpers/PCGAsync.h"
 #include "TextureResource.h"
 
+#include UE_INLINE_GENERATED_CPP_BY_NAME(PCGTextureData)
+
 namespace PCGTextureSampling
 {
 	template<typename ValueType>
@@ -338,4 +340,4 @@ UPCGSpatialData* UPCGTextureData::CopyInternal() const
 	NewTextureData->Texture = Texture;
 
 	return NewTextureData;
-}
\ No newline at end of file
+}
diff --git a/Engine/Plugins/Experimental/PCG/Source/PCG/Private/Data/PCGUnionData.cpp b/Engine/Plugins/Experimental/PCG/Source/PCG/Private/Data/PCGUnionData.cpp
--- a/Engine/Plugins/Experimental/PCG/Source/PCG/Private/Data/PCGUnionData.cpp
+++ b/Engine/Plugins/Experimental/PCG/Source/PCG/Private/Data/PCGUnionData.cpp
@@ -5,6 +5,8 @@
 #include "Helpers/PCGAsync.h"
 #include "PCGHelpers.h"
 
+#include UE_INLINE_GENERATED_CPP_BY_NAME(PCGUnionData)
+
 namespace PCGUnionDataMaths
 {
 	float ComputeDensity(float InDensityToUpdate, float InOtherDensity, EPCGUnionDensityFunction DensityFunction)
@@ -344,3 +346,4 @@ UPCGSpatialData* UPCGUnionData::CopyInternal() const
 
 	return NewUnionData;
 }
+
diff --git a/Engine/Plugins/Experimental/PCG/Source/PCG/Private/Data/PCGVolumeData.cpp b/Engine/Plugins/Experimental/PCG/Source/PCG/Private/Data/PCGVolumeData.cpp
--- a/Engine/Plugins/Experimental/PCG/Source/PCG/Private/Data/PCGVolumeData.cpp
+++ b/Engine/Plugins/Experimental/PCG/Source/PCG/Private/Data/PCGVolumeData.cpp
@@ -8,6 +8,8 @@
 #include "Components/BrushComponent.h"
 #include "GameFramework/Volume.h"
 
+#include UE_INLINE_GENERATED_CPP_BY_NAME(PCGVolumeData)
+
 void UPCGVolumeData::Initialize(AVolume* InVolume, AActor* InTargetActor)
 {
 	check(InVolume);
@@ -101,4 +103,4 @@ UPCGSpatialData* UPCGVolumeData::CopyInternal() const
 	CopyBaseVolumeData(NewVolumeData);
 
 	return NewVolumeData;
-}
\ No newline at end of file
+}
diff --git a/Engine/Plugins/Experimental/PCG/Source/PCG/Private/Data/PCGWorldData.cpp b/Engine/Plugins/Experimental/PCG/Source/PCG/Private/Data/PCGWorldData.cpp
--- a/Engine/Plugins/Experimental/PCG/Source/PCG/Private/Data/PCGWorldData.cpp
+++ b/Engine/Plugins/Experimental/PCG/Source/PCG/Private/Data/PCGWorldData.cpp
@@ -12,6 +12,8 @@
 #include "Components/BrushComponent.h"
 #include "Engine/World.h"
 
+#include UE_INLINE_GENERATED_CPP_BY_NAME(PCGWorldData)
+
 void UPCGWorldVolumetricData::Initialize(UWorld* InWorld, const FBox& InBounds)
 {
 	Super::Initialize(InBounds, nullptr);
@@ -259,4 +261,4 @@ UPCGSpatialData* UPCGWorldRayHitData::CopyInternal() const
 	NewData->QueryParams = QueryParams;
 
 	return NewData;
-}
\ No newline at end of file
+}
diff --git a/Engine/Plugins/Experimental/PCG/Source/PCG/Private/Elements/Metadata/PCGMetadataBitwiseOpElement.cpp b/Engine/Plugins/Experimental/PCG/Source/PCG/Private/Elements/Metadata/PCGMetadataBitwiseOpElement.cpp
--- a/Engine/Plugins/Experimental/PCG/Source/PCG/Private/Elements/Metadata/PCGMetadataBitwiseOpElement.cpp
+++ b/Engine/Plugins/Experimental/PCG/Source/PCG/Private/Elements/Metadata/PCGMetadataBitwiseOpElement.cpp
@@ -11,6 +11,8 @@
 #include "Metadata/PCGMetadataEntryKeyIterator.h"
 #include "Elements/Metadata/PCGMetadataElementCommon.h"
 
+#include UE_INLINE_GENERATED_CPP_BY_NAME(PCGMetadataBitwiseOpElement)
+
 namespace PCGMetadataBitwiseSettings
 {
 	inline int64 UnaryOp(const int64& Value)
@@ -138,4 +140,4 @@ bool FPCGMetadataBitwiseElement::DoOperation(FOperationData& OperationData) cons
 	}
 
 	return true;
-}
\ No newline at end of file
+}
diff --git a/Engine/Plugins/Experimental/PCG/Source/PCG/Private/Elements/Metadata/PCGMetadataBooleanOpElement.cpp b/Engine/Plugins/Experimental/PCG/Source/PCG/Private/Elements/Metadata/PCGMetadataBooleanOpElement.cpp
--- a/Engine/Plugins/Experimental/PCG/Source/PCG/Private/Elements/Metadata/PCGMetadataBooleanOpElement.cpp
+++ b/Engine/Plugins/Experimental/PCG/Source/PCG/Private/Elements/Metadata/PCGMetadataBooleanOpElement.cpp
@@ -11,6 +11,8 @@
 #include "Metadata/PCGMetadataEntryKeyIterator.h"
 #include "Elements/Metadata/PCGMetadataElementCommon.h"
 
+#include UE_INLINE_GENERATED_CPP_BY_NAME(PCGMetadataBooleanOpElement)
+
 namespace PCGMetadataBooleanSettings
 {
 	inline bool UnaryOp(const bool& Value)
@@ -138,4 +140,4 @@ bool FPCGMetadataBooleanElement::DoOperation(FOperationData& OperationData) cons
 	}
 
 	return true;
-}
\ No newline at end of file
+}
diff --git a/Engine/Plugins/Experimental/PCG/Source/PCG/Private/Elements/Metadata/PCGMetadataBreakTransform.cpp b/Engine/Plugins/Experimental/PCG/Source/PCG/Private/Elements/Metadata/PCGMetadataBreakTransform.cpp
--- a/Engine/Plugins/Experimental/PCG/Source/PCG/Private/Elements/Metadata/PCGMetadataBreakTransform.cpp
+++ b/Engine/Plugins/Experimental/PCG/Source/PCG/Private/Elements/Metadata/PCGMetadataBreakTransform.cpp
@@ -10,6 +10,8 @@
 #include "Metadata/PCGMetadataAttribute.h"
 #include "Metadata/PCGMetadataAttributeTpl.h"
 
+#include UE_INLINE_GENERATED_CPP_BY_NAME(PCGMetadataBreakTransform)
+
 namespace PCGMetadataBreakTransformSettings
 {
 	inline FVector GetTranslation(const FTransform& Transform)
@@ -119,3 +121,4 @@ bool FPCGMetadataBreakTransformElement::DoOperation(FOperationData& OperationDat
 		PCGMetadataBreakTransformSettings::GetRotation,
 		PCGMetadataBreakTransformSettings::GetScale);
 }
+
diff --git a/Engine/Plugins/Experimental/PCG/Source/PCG/Private/Elements/Metadata/PCGMetadataBreakVector.cpp b/Engine/Plugins/Experimental/PCG/Source/PCG/Private/Elements/Metadata/PCGMetadataBreakVector.cpp
--- a/Engine/Plugins/Experimental/PCG/Source/PCG/Private/Elements/Metadata/PCGMetadataBreakVector.cpp
+++ b/Engine/Plugins/Experimental/PCG/Source/PCG/Private/Elements/Metadata/PCGMetadataBreakVector.cpp
@@ -10,6 +10,8 @@
 #include "Metadata/PCGMetadataAttribute.h"
 #include "Metadata/PCGMetadataAttributeTpl.h"
 
+#include UE_INLINE_GENERATED_CPP_BY_NAME(PCGMetadataBreakVector)
+
 namespace PCGMetadataBreakVectorSettings
 {
 	template <typename InType>
@@ -187,3 +189,4 @@ bool FPCGMetadataBreakVectorElement::DoOperation(FOperationData& OperationData)
 
 	return PCGMetadataAttribute::CallbackWithRightType(OperationData.MostComplexInputType, BreakFunc);
 }
+
diff --git a/Engine/Plugins/Experimental/PCG/Source/PCG/Private/Elements/Metadata/PCGMetadataCompareOpElement.cpp b/Engine/Plugins/Experimental/PCG/Source/PCG/Private/Elements/Metadata/PCGMetadataCompareOpElement.cpp
--- a/Engine/Plugins/Experimental/PCG/Source/PCG/Private/Elements/Metadata/PCGMetadataCompareOpElement.cpp
+++ b/Engine/Plugins/Experimental/PCG/Source/PCG/Private/Elements/Metadata/PCGMetadataCompareOpElement.cpp
@@ -9,6 +9,8 @@
 #include "Metadata/PCGMetadataEntryKeyIterator.h"
 #include "Elements/Metadata/PCGMetadataElementCommon.h"
 
+#include UE_INLINE_GENERATED_CPP_BY_NAME(PCGMetadataCompareOpElement)
+
 namespace PCGMetadataCompareSettings
 {
 	template <typename T>
@@ -166,4 +168,4 @@ bool FPCGMetadataCompareElement::DoOperation(FOperationData& OperationData) cons
 	PCGMetadataAttribute::CallbackWithRightType(OperationData.MostComplexInputType, CompareFunc);
 
 	return true;
-}
\ No newline at end of file
+}
diff --git a/Engine/Plugins/Experimental/PCG/Source/PCG/Private/Elements/Metadata/PCGMetadataElement.cpp b/Engine/Plugins/Experimental/PCG/Source/PCG/Private/Elements/Metadata/PCGMetadataElement.cpp
--- a/Engine/Plugins/Experimental/PCG/Source/PCG/Private/Elements/Metadata/PCGMetadataElement.cpp
+++ b/Engine/Plugins/Experimental/PCG/Source/PCG/Private/Elements/Metadata/PCGMetadataElement.cpp
@@ -7,6 +7,8 @@
 #include "Metadata/PCGMetadataAttribute.h"
 #include "Metadata/PCGMetadataAttributeTpl.h"
 
+#include UE_INLINE_GENERATED_CPP_BY_NAME(PCGMetadataElement)
+
 namespace PCGMetadataOperations
 {
 	// Getter and setter for properties.
@@ -422,4 +424,4 @@ bool FPCGMetadataOperationElement::ExecuteInternal(FPCGContext* Context) const
 	}
 
 	return true;
-}
\ No newline at end of file
+}
diff --git a/Engine/Plugins/Experimental/PCG/Source/PCG/Private/Elements/Metadata/PCGMetadataMakeTransform.cpp b/Engine/Plugins/Experimental/PCG/Source/PCG/Private/Elements/Metadata/PCGMetadataMakeTransform.cpp
--- a/Engine/Plugins/Experimental/PCG/Source/PCG/Private/Elements/Metadata/PCGMetadataMakeTransform.cpp
+++ b/Engine/Plugins/Experimental/PCG/Source/PCG/Private/Elements/Metadata/PCGMetadataMakeTransform.cpp
@@ -6,6 +6,8 @@
 #include "Metadata/PCGMetadataAttributeTpl.h"
 #include "Elements/Metadata/PCGMetadataElementCommon.h"
 
+#include UE_INLINE_GENERATED_CPP_BY_NAME(PCGMetadataMakeTransform)
+
 namespace PCGMetadataMakeTransformSettings
 {
 	template <typename VectorType = FVector4>
@@ -142,4 +144,4 @@ bool FPCGMetadataMakeTransformElement::DoOperation(FOperationData& OperationData
 	};
 
 	return PCGMetadataAttribute::CallbackWithRightType(OperationData.MostComplexInputType, TransformFunc);
-}
\ No newline at end of file
+}
diff --git a/Engine/Plugins/Experimental/PCG/Source/PCG/Private/Elements/Metadata/PCGMetadataMakeVector.cpp b/Engine/Plugins/Experimental/PCG/Source/PCG/Private/Elements/Metadata/PCGMetadataMakeVector.cpp
--- a/Engine/Plugins/Experimental/PCG/Source/PCG/Private/Elements/Metadata/PCGMetadataMakeVector.cpp
+++ b/Engine/Plugins/Experimental/PCG/Source/PCG/Private/Elements/Metadata/PCGMetadataMakeVector.cpp
@@ -6,6 +6,8 @@
 #include "Metadata/PCGMetadataAttributeTpl.h"
 #include "Elements/Metadata/PCGMetadataElementCommon.h"
 
+#include UE_INLINE_GENERATED_CPP_BY_NAME(PCGMetadataMakeVector)
+
 namespace PCGMetadataMakeVectorSettings
 {
 	FVector2D MakeVector2(const double& X, const double& Y)
@@ -241,4 +243,4 @@ bool FPCGMetadataMakeVectorElement::DoOperation(FOperationData& OperationData) c
 	}
 
 	return true;
-}
\ No newline at end of file
+}
diff --git a/Engine/Plugins/Experimental/PCG/Source/PCG/Private/Elements/Metadata/PCGMetadataMathsOpElement.cpp b/Engine/Plugins/Experimental/PCG/Source/PCG/Private/Elements/Metadata/PCGMetadataMathsOpElement.cpp
--- a/Engine/Plugins/Experimental/PCG/Source/PCG/Private/Elements/Metadata/PCGMetadataMathsOpElement.cpp
+++ b/Engine/Plugins/Experimental/PCG/Source/PCG/Private/Elements/Metadata/PCGMetadataMathsOpElement.cpp
@@ -13,6 +13,8 @@
 
 #include "Elements/Metadata/PCGMetadataMaths.inl"
 
+#include UE_INLINE_GENERATED_CPP_BY_NAME(PCGMetadataMathsOpElement)
+
 namespace PCGMetadataMathsSettings
 {
 	inline constexpr bool IsUnaryOp(EPCGMedadataMathsOperation Operation)
@@ -300,4 +302,4 @@ bool FPCGMetadataMathsElement::DoOperation(FOperationData& OperationData) const
 	PCGMetadataAttribute::CallbackWithRightType(OperationData.OutputType, MathFunc);
 
 	return true;
-}
\ No newline at end of file
+}
diff --git a/Engine/Plugins/Experimental/PCG/Source/PCG/Private/Elements/Metadata/PCGMetadataOpElementBase.cpp b/Engine/Plugins/Experimental/PCG/Source/PCG/Private/Elements/Metadata/PCGMetadataOpElementBase.cpp
--- a/Engine/Plugins/Experimental/PCG/Source/PCG/Private/Elements/Metadata/PCGMetadataOpElementBase.cpp
+++ b/Engine/Plugins/Experimental/PCG/Source/PCG/Private/Elements/Metadata/PCGMetadataOpElementBase.cpp
@@ -9,6 +9,8 @@
 #include "Helpers/PCGSettingsHelpers.h"
 #include "Metadata/PCGMetadata.h"
 
+#include UE_INLINE_GENERATED_CPP_BY_NAME(PCGMetadataOpElementBase)
+
 void UPCGMetadataSettingsBase::PostLoad()
 {
 	Super::PostLoad();
@@ -303,3 +305,4 @@ bool FPCGMetadataElementBase::ExecuteInternal(FPCGContext* Context) const
 
 	return true;
 }
+
diff --git a/Engine/Plugins/Experimental/PCG/Source/PCG/Private/Elements/Metadata/PCGMetadataPartition.cpp b/Engine/Plugins/Experimental/PCG/Source/PCG/Private/Elements/Metadata/PCGMetadataPartition.cpp
--- a/Engine/Plugins/Experimental/PCG/Source/PCG/Private/Elements/Metadata/PCGMetadataPartition.cpp
+++ b/Engine/Plugins/Experimental/PCG/Source/PCG/Private/Elements/Metadata/PCGMetadataPartition.cpp
@@ -5,6 +5,8 @@
 #include "Data/PCGPointData.h"
 #include "Helpers/PCGSettingsHelpers.h"
 
+#include UE_INLINE_GENERATED_CPP_BY_NAME(PCGMetadataPartition)
+
 FPCGElementPtr UPCGMetadataPartitionSettings::CreateElement() const
 {
 	return MakeShared<FPCGMetadataPartitionElement>();
@@ -123,4 +125,4 @@ bool FPCGMetadataPartitionElement::ExecuteInternal(FPCGContext* Context) const
 	}
 
 	return true;
-}
\ No newline at end of file
+}
diff --git a/Engine/Plugins/Experimental/PCG/Source/PCG/Private/Elements/Metadata/PCGMetadataRenameElement.cpp b/Engine/Plugins/Experimental/PCG/Source/PCG/Private/Elements/Metadata/PCGMetadataRenameElement.cpp
--- a/Engine/Plugins/Experimental/PCG/Source/PCG/Private/Elements/Metadata/PCGMetadataRenameElement.cpp
+++ b/Engine/Plugins/Experimental/PCG/Source/PCG/Private/Elements/Metadata/PCGMetadataRenameElement.cpp
@@ -6,6 +6,8 @@
 #include "Metadata/PCGMetadata.h"
 #include "Metadata/PCGMetadataAttribute.h"
 
+#include UE_INLINE_GENERATED_CPP_BY_NAME(PCGMetadataRenameElement)
+
 TArray<FPCGPinProperties> UPCGMetadataRenameSettings::OutputPinProperties() const
 {
 	TArray<FPCGPinProperties> PinProperties;
@@ -115,4 +117,4 @@ bool FPCGMetadataRenameElement::ExecuteInternal(FPCGContext* Context) const
 	Outputs.Append(Context->InputData.GetAllSettings());
 
 	return true;
-}
\ No newline at end of file
+}
diff --git a/Engine/Plugins/Experimental/PCG/Source/PCG/Private/Elements/Metadata/PCGMetadataRotatorOpElement.cpp b/Engine/Plugins/Experimental/PCG/Source/PCG/Private/Elements/Metadata/PCGMetadataRotatorOpElement.cpp
--- a/Engine/Plugins/Experimental/PCG/Source/PCG/Private/Elements/Metadata/PCGMetadataRotatorOpElement.cpp
+++ b/Engine/Plugins/Experimental/PCG/Source/PCG/Private/Elements/Metadata/PCGMetadataRotatorOpElement.cpp
@@ -9,6 +9,8 @@
 #include "Metadata/PCGMetadataEntryKeyIterator.h"
 #include "Elements/Metadata/PCGMetadataElementCommon.h"
 
+#include UE_INLINE_GENERATED_CPP_BY_NAME(PCGMetadataRotatorOpElement)
+
 // Taken from Kismet Math library
 FRotator PCGMetadataRotatorHelpers::RLerp(const FRotator& A, const FRotator& B, double Alpha, bool bShortestPath)
 {
@@ -266,4 +268,4 @@ bool FPCGMetadataRotatorElement::DoOperation(FOperationData& OperationData) cons
 	PCGMetadataAttribute::CallbackWithRightType(OperationData.OutputType, RotatorFunc);
 
 	return true;
-}
\ No newline at end of file
+}
diff --git a/Engine/Plugins/Experimental/PCG/Source/PCG/Private/Elements/Metadata/PCGMetadataTransformOpElement.cpp b/Engine/Plugins/Experimental/PCG/Source/PCG/Private/Elements/Metadata/PCGMetadataTransformOpElement.cpp
--- a/Engine/Plugins/Experimental/PCG/Source/PCG/Private/Elements/Metadata/PCGMetadataTransformOpElement.cpp
+++ b/Engine/Plugins/Experimental/PCG/Source/PCG/Private/Elements/Metadata/PCGMetadataTransformOpElement.cpp
@@ -12,6 +12,8 @@
 
 #include "Math/DualQuat.h"
 
+#include UE_INLINE_GENERATED_CPP_BY_NAME(PCGMetadataTransformOpElement)
+
 namespace PCGMetadataTransfromSettings
 {
 	inline bool IsUnaryOp(EPCGMedadataTransformOperation Operation)
@@ -215,4 +217,4 @@ bool FPCGMetadataTransformElement::DoOperation(FOperationData& OperationData) co
 	}
 
 	return true;
-}
\ No newline at end of file
+}
diff --git a/Engine/Plugins/Experimental/PCG/Source/PCG/Private/Elements/Metadata/PCGMetadataTrigOpElement.cpp b/Engine/Plugins/Experimental/PCG/Source/PCG/Private/Elements/Metadata/PCGMetadataTrigOpElement.cpp
--- a/Engine/Plugins/Experimental/PCG/Source/PCG/Private/Elements/Metadata/PCGMetadataTrigOpElement.cpp
+++ b/Engine/Plugins/Experimental/PCG/Source/PCG/Private/Elements/Metadata/PCGMetadataTrigOpElement.cpp
@@ -9,6 +9,8 @@
 #include "Metadata/PCGMetadataEntryKeyIterator.h"
 #include "Elements/Metadata/PCGMetadataElementCommon.h"
 
+#include UE_INLINE_GENERATED_CPP_BY_NAME(PCGMetadataTrigOpElement)
+
 namespace PCGMetadataTrigSettings
 {
 	template <typename OutType>
@@ -148,4 +150,4 @@ bool FPCGMetadataTrigElement::DoOperation(FOperationData& OperationData) const
 	}
 
 	return true;
-}
\ No newline at end of file
+}
diff --git a/Engine/Plugins/Experimental/PCG/Source/PCG/Private/Elements/Metadata/PCGMetadataVectorOpElement.cpp b/Engine/Plugins/Experimental/PCG/Source/PCG/Private/Elements/Metadata/PCGMetadataVectorOpElement.cpp
--- a/Engine/Plugins/Experimental/PCG/Source/PCG/Private/Elements/Metadata/PCGMetadataVectorOpElement.cpp
+++ b/Engine/Plugins/Experimental/PCG/Source/PCG/Private/Elements/Metadata/PCGMetadataVectorOpElement.cpp
@@ -11,6 +11,8 @@
 #include "Metadata/PCGMetadataEntryKeyIterator.h"
 #include "Elements/Metadata/PCGMetadataElementCommon.h"
 
+#include UE_INLINE_GENERATED_CPP_BY_NAME(PCGMetadataVectorOpElement)
+
 namespace PCGMetadataVectorConstants
 {
 	const FName AxisLabel = TEXT("Axis");
@@ -436,4 +438,4 @@ bool FPCGMetadataVectorElement::DoOperation(FOperationData& OperationData) const
 	PCGMetadataAttribute::CallbackWithRightType(OperationData.MostComplexInputType, VectorFunc);
 
 	return true;
-}
\ No newline at end of file
+}
diff --git a/Engine/Plugins/Experimental/PCG/Source/PCG/Private/Elements/PCGActorSelector.cpp b/Engine/Plugins/Experimental/PCG/Source/PCG/Private/Elements/PCGActorSelector.cpp
--- a/Engine/Plugins/Experimental/PCG/Source/PCG/Private/Elements/PCGActorSelector.cpp
+++ b/Engine/Plugins/Experimental/PCG/Source/PCG/Private/Elements/PCGActorSelector.cpp
@@ -5,6 +5,8 @@
 
 #include "Kismet/GameplayStatics.h"
 
+#include UE_INLINE_GENERATED_CPP_BY_NAME(PCGActorSelector)
+
 namespace PCGActorSelector
 {
 	// Need to pass a pointer of pointer to the found actor. The lambda will capture this pointer and modify its value when an actor is found.
@@ -170,4 +172,4 @@ namespace PCGActorSelector
 		TArray<AActor*> Actors = FindActors(Settings, World, Self);
 		return Actors.IsEmpty() ? nullptr : Actors[0];
 	}
-}
\ No newline at end of file
+}
diff --git a/Engine/Plugins/Experimental/PCG/Source/PCG/Private/Elements/PCGAttributeFilterElement.cpp b/Engine/Plugins/Experimental/PCG/Source/PCG/Private/Elements/PCGAttributeFilterElement.cpp
--- a/Engine/Plugins/Experimental/PCG/Source/PCG/Private/Elements/PCGAttributeFilterElement.cpp
+++ b/Engine/Plugins/Experimental/PCG/Source/PCG/Private/Elements/PCGAttributeFilterElement.cpp
@@ -8,6 +8,8 @@
 #include "Metadata/PCGMetadata.h"
 #include "Metadata/PCGMetadataAttribute.h"
 
+#include UE_INLINE_GENERATED_CPP_BY_NAME(PCGAttributeFilterElement)
+
 namespace PCGAttributeFilterConstants
 {
 	const FName NodeName = TEXT("FilterAttribute");
@@ -106,3 +108,4 @@ bool FPCGAttributeFilterElement::ExecuteInternal(FPCGContext* Context) const
 
 	return true;
 }
+
diff --git a/Engine/Plugins/Experimental/PCG/Source/PCG/Private/Elements/PCGAttributeGetFromPointIndexElement.cpp b/Engine/Plugins/Experimental/PCG/Source/PCG/Private/Elements/PCGAttributeGetFromPointIndexElement.cpp
--- a/Engine/Plugins/Experimental/PCG/Source/PCG/Private/Elements/PCGAttributeGetFromPointIndexElement.cpp
+++ b/Engine/Plugins/Experimental/PCG/Source/PCG/Private/Elements/PCGAttributeGetFromPointIndexElement.cpp
@@ -10,6 +10,8 @@
 #include "Metadata/PCGMetadataAttribute.h"
 #include "Helpers/PCGSettingsHelpers.h"
 
+#include UE_INLINE_GENERATED_CPP_BY_NAME(PCGAttributeGetFromPointIndexElement)
+
 #if WITH_EDITOR
 FName UPCGAttributeGetFromPointIndexSettings::GetDefaultNodeName() const
 {
@@ -118,4 +120,4 @@ bool FPCGAttributeGetFromPointIndexElement::ExecuteInternal(FPCGContext* Context
 	}
 
 	return true;
-}
\ No newline at end of file
+}
diff --git a/Engine/Plugins/Experimental/PCG/Source/PCG/Private/Elements/PCGAttributeReduceElement.cpp b/Engine/Plugins/Experimental/PCG/Source/PCG/Private/Elements/PCGAttributeReduceElement.cpp
--- a/Engine/Plugins/Experimental/PCG/Source/PCG/Private/Elements/PCGAttributeReduceElement.cpp
+++ b/Engine/Plugins/Experimental/PCG/Source/PCG/Private/Elements/PCGAttributeReduceElement.cpp
@@ -10,6 +10,8 @@
 #include "Metadata/PCGMetadataEntryKeyIterator.h"
 #include "Helpers/PCGSettingsHelpers.h"
 
+#include UE_INLINE_GENERATED_CPP_BY_NAME(PCGAttributeReduceElement)
+
 namespace PCGAttributeReduceElement
 {
 	template <typename T>
@@ -232,4 +234,4 @@ bool FPCGAttributeReduceElement::ExecuteInternal(FPCGContext* Context) const
 	Output.Data = OutputParamData;
 
 	return true;
-}
\ No newline at end of file
+}
diff --git a/Engine/Plugins/Experimental/PCG/Source/PCG/Private/Elements/PCGAttributeSelectElement.cpp b/Engine/Plugins/Experimental/PCG/Source/PCG/Private/Elements/PCGAttributeSelectElement.cpp
--- a/Engine/Plugins/Experimental/PCG/Source/PCG/Private/Elements/PCGAttributeSelectElement.cpp
+++ b/Engine/Plugins/Experimental/PCG/Source/PCG/Private/Elements/PCGAttributeSelectElement.cpp
@@ -10,6 +10,8 @@
 #include "Metadata/PCGMetadataEntryKeyIterator.h"
 #include "Helpers/PCGSettingsHelpers.h"
 
+#include UE_INLINE_GENERATED_CPP_BY_NAME(PCGAttributeSelectElement)
+
 namespace PCGAttributeSelectElement
 {
 	// Of course, Dot doesn't exist for FVector4...
@@ -408,4 +410,4 @@ bool FPCGAttributeSelectElement::ExecuteInternal(FPCGContext* Context) const
 	}
 
 	return true;
-}
\ No newline at end of file
+}
diff --git a/Engine/Plugins/Experimental/PCG/Source/PCG/Private/Elements/PCGAttributeTransferElement.cpp b/Engine/Plugins/Experimental/PCG/Source/PCG/Private/Elements/PCGAttributeTransferElement.cpp
--- a/Engine/Plugins/Experimental/PCG/Source/PCG/Private/Elements/PCGAttributeTransferElement.cpp
+++ b/Engine/Plugins/Experimental/PCG/Source/PCG/Private/Elements/PCGAttributeTransferElement.cpp
@@ -8,6 +8,8 @@
 #include "Metadata/PCGMetadata.h"
 #include "Metadata/PCGMetadataAttribute.h"
 
+#include UE_INLINE_GENERATED_CPP_BY_NAME(PCGAttributeTransferElement)
+
 #if WITH_EDITOR
 FName UPCGAttributeTransferSettings::GetDefaultNodeName() const
 {
@@ -174,4 +176,4 @@ UPCGPointData* FPCGAttributeTransferElement::TransferPointToPoint(FPCGContext* C
 	}
 
 	return NewPointData;
-}
\ No newline at end of file
+}
diff --git a/Engine/Plugins/Experimental/PCG/Source/PCG/Private/Elements/PCGCollapseElement.cpp b/Engine/Plugins/Experimental/PCG/Source/PCG/Private/Elements/PCGCollapseElement.cpp
--- a/Engine/Plugins/Experimental/PCG/Source/PCG/Private/Elements/PCGCollapseElement.cpp
+++ b/Engine/Plugins/Experimental/PCG/Source/PCG/Private/Elements/PCGCollapseElement.cpp
@@ -5,6 +5,8 @@
 #include "Data/PCGSpatialData.h"
 #include "Data/PCGPointData.h"
 
+#include UE_INLINE_GENERATED_CPP_BY_NAME(PCGCollapseElement)
+
 FPCGElementPtr UPCGCollapseSettings::CreateElement() const
 {
 	return MakeShared<FPCGCollapseElement>();
@@ -31,4 +33,4 @@ bool FPCGCollapseElement::ExecuteInternal(FPCGContext* Context) const
 	}
 
 	return true;
-}
\ No newline at end of file
+}
diff --git a/Engine/Plugins/Experimental/PCG/Source/PCG/Private/Elements/PCGCopyPoints.cpp b/Engine/Plugins/Experimental/PCG/Source/PCG/Private/Elements/PCGCopyPoints.cpp
--- a/Engine/Plugins/Experimental/PCG/Source/PCG/Private/Elements/PCGCopyPoints.cpp
+++ b/Engine/Plugins/Experimental/PCG/Source/PCG/Private/Elements/PCGCopyPoints.cpp
@@ -8,6 +8,8 @@
 #include "Helpers/PCGAsync.h"
 #include "Helpers/PCGSettingsHelpers.h"
 
+#include UE_INLINE_GENERATED_CPP_BY_NAME(PCGCopyPoints)
+
 TArray<FPCGPinProperties> UPCGCopyPointsSettings::InputPinProperties() const
 {
 	TArray<FPCGPinProperties> PinProperties;
@@ -219,3 +221,4 @@ bool FPCGCopyPointsElement::ExecuteInternal(FPCGContext* Context) const
 
 	return true;
 }
+
diff --git a/Engine/Plugins/Experimental/PCG/Source/PCG/Private/Elements/PCGCreateAttribute.cpp b/Engine/Plugins/Experimental/PCG/Source/PCG/Private/Elements/PCGCreateAttribute.cpp
--- a/Engine/Plugins/Experimental/PCG/Source/PCG/Private/Elements/PCGCreateAttribute.cpp
+++ b/Engine/Plugins/Experimental/PCG/Source/PCG/Private/Elements/PCGCreateAttribute.cpp
@@ -10,6 +10,8 @@
 #include "Metadata/PCGMetadata.h"
 #include "Metadata/PCGMetadataAttribute.h"
 
+#include UE_INLINE_GENERATED_CPP_BY_NAME(PCGCreateAttribute)
+
 namespace PCGCreateAttributeConstants
 {
 	const FName NodeName = TEXT("CreateAttribute");
@@ -304,3 +306,4 @@ PCGMetadataEntryKey FPCGCreateAttributeElement::SetAttribute(const UPCGCreateAtt
 
 	return PCGCreateAttributeElement::Dispatcher(Settings, Params, SetAttribute);
 }
+
diff --git a/Engine/Plugins/Experimental/PCG/Source/PCG/Private/Elements/PCGDataFromActor.cpp b/Engine/Plugins/Experimental/PCG/Source/PCG/Private/Elements/PCGDataFromActor.cpp
--- a/Engine/Plugins/Experimental/PCG/Source/PCG/Private/Elements/PCGDataFromActor.cpp
+++ b/Engine/Plugins/Experimental/PCG/Source/PCG/Private/Elements/PCGDataFromActor.cpp
@@ -8,6 +8,8 @@
 #include "GameFramework/Actor.h"
 #include "UObject/Package.h"
 
+#include UE_INLINE_GENERATED_CPP_BY_NAME(PCGDataFromActor)
+
 FPCGElementPtr UPCGDataFromActorSettings::CreateElement() const
 {
 	return MakeShared<FPCGDataFromActorElement>();
@@ -156,4 +158,4 @@ void FPCGDataFromActorElement::ProcessActor(FPCGContext* Context, const UPCGData
 			}
 		}
 	}
-}
\ No newline at end of file
+}
diff --git a/Engine/Plugins/Experimental/PCG/Source/PCG/Private/Elements/PCGDataTableRowToParamData.cpp b/Engine/Plugins/Experimental/PCG/Source/PCG/Private/Elements/PCGDataTableRowToParamData.cpp
--- a/Engine/Plugins/Experimental/PCG/Source/PCG/Private/Elements/PCGDataTableRowToParamData.cpp
+++ b/Engine/Plugins/Experimental/PCG/Source/PCG/Private/Elements/PCGDataTableRowToParamData.cpp
@@ -6,6 +6,8 @@
 #include "Metadata/PCGMetadata.h"
 #include "Engine/DataTable.h"
 
+#include UE_INLINE_GENERATED_CPP_BY_NAME(PCGDataTableRowToParamData)
+
 FPCGElementPtr UPCGDataTableRowToParamDataSettings::CreateElement() const
 {
 	return MakeShared<FPCGDataTableRowToParamData>();
@@ -69,3 +71,4 @@ bool FPCGDataTableRowToParamData::ExecuteInternal(FPCGContext* Context) const
 
 	return true;
 }
+
diff --git a/Engine/Plugins/Experimental/PCG/Source/PCG/Private/Elements/PCGDebugElement.cpp b/Engine/Plugins/Experimental/PCG/Source/PCG/Private/Elements/PCGDebugElement.cpp
--- a/Engine/Plugins/Experimental/PCG/Source/PCG/Private/Elements/PCGDebugElement.cpp
+++ b/Engine/Plugins/Experimental/PCG/Source/PCG/Private/Elements/PCGDebugElement.cpp
@@ -11,6 +11,8 @@
 #include "Engine/StaticMesh.h"
 #include "Materials/MaterialInterface.h"
 
+#include UE_INLINE_GENERATED_CPP_BY_NAME(PCGDebugElement)
+
 namespace PCGDebugElement
 {
 	void ExecuteDebugDisplay(FPCGContext* Context)
@@ -193,4 +195,4 @@ bool FPCGDebugElement::ExecuteInternal(FPCGContext* Context) const
 	
 	Context->OutputData = Context->InputData;
 	return true;
-}
\ No newline at end of file
+}
diff --git a/Engine/Plugins/Experimental/PCG/Source/PCG/Private/Elements/PCGDensityFilter.cpp b/Engine/Plugins/Experimental/PCG/Source/PCG/Private/Elements/PCGDensityFilter.cpp
--- a/Engine/Plugins/Experimental/PCG/Source/PCG/Private/Elements/PCGDensityFilter.cpp
+++ b/Engine/Plugins/Experimental/PCG/Source/PCG/Private/Elements/PCGDensityFilter.cpp
@@ -5,6 +5,8 @@
 #include "Helpers/PCGAsync.h"
 #include "Helpers/PCGSettingsHelpers.h"
 
+#include UE_INLINE_GENERATED_CPP_BY_NAME(PCGDensityFilter)
+
 FPCGElementPtr UPCGDensityFilterSettings::CreateElement() const
 {
 	return MakeShared<FPCGDensityFilterElement>();
@@ -107,4 +109,4 @@ bool FPCGDensityFilterElement::ExecuteInternal(FPCGContext* Context) const
 	}
 
 	return true;
-}
\ No newline at end of file
+}
diff --git a/Engine/Plugins/Experimental/PCG/Source/PCG/Private/Elements/PCGDensityNoise.cpp b/Engine/Plugins/Experimental/PCG/Source/PCG/Private/Elements/PCGDensityNoise.cpp
--- a/Engine/Plugins/Experimental/PCG/Source/PCG/Private/Elements/PCGDensityNoise.cpp
+++ b/Engine/Plugins/Experimental/PCG/Source/PCG/Private/Elements/PCGDensityNoise.cpp
@@ -7,6 +7,8 @@
 
 #include "Math/RandomStream.h"
 
+#include UE_INLINE_GENERATED_CPP_BY_NAME(PCGDensityNoise)
+
 UPCGDensityNoiseSettings::UPCGDensityNoiseSettings()
 {
 	bUseSeed = true;
@@ -79,3 +81,4 @@ bool FPCGDensityNoiseElement::ExecuteInternal(FPCGContext* Context) const
 
 	return true;
 }
+
diff --git a/Engine/Plugins/Experimental/PCG/Source/PCG/Private/Elements/PCGDensityRemap.cpp b/Engine/Plugins/Experimental/PCG/Source/PCG/Private/Elements/PCGDensityRemap.cpp
--- a/Engine/Plugins/Experimental/PCG/Source/PCG/Private/Elements/PCGDensityRemap.cpp
+++ b/Engine/Plugins/Experimental/PCG/Source/PCG/Private/Elements/PCGDensityRemap.cpp
@@ -7,6 +7,8 @@
 #include "Helpers/PCGSettingsHelpers.h"
 #include "Math/RandomStream.h"
 
+#include UE_INLINE_GENERATED_CPP_BY_NAME(PCGDensityRemap)
+
 UPCGLinearDensityRemapSettings::UPCGLinearDensityRemapSettings()
 {
 	bUseSeed = true;
@@ -91,4 +93,4 @@ bool FPCGLinearDensityRemapElement::ExecuteInternal(FPCGContext* Context) const
 	}
 
 	return true;
-}
\ No newline at end of file
+}
diff --git a/Engine/Plugins/Experimental/PCG/Source/PCG/Private/Elements/PCGDensityRemapElement.cpp b/Engine/Plugins/Experimental/PCG/Source/PCG/Private/Elements/PCGDensityRemapElement.cpp
--- a/Engine/Plugins/Experimental/PCG/Source/PCG/Private/Elements/PCGDensityRemapElement.cpp
+++ b/Engine/Plugins/Experimental/PCG/Source/PCG/Private/Elements/PCGDensityRemapElement.cpp
@@ -7,6 +7,8 @@
 
 #include "Math/RandomStream.h"
 
+#include UE_INLINE_GENERATED_CPP_BY_NAME(PCGDensityRemapElement)
+
 FPCGElementPtr UPCGDensityRemapSettings::CreateElement() const
 {
 	return MakeShared<FPCGDensityRemapElement>();
@@ -70,3 +72,4 @@ bool FPCGDensityRemapElement::ExecuteInternal(FPCGContext* Context) const
 
 	return true;
 }
+
diff --git a/Engine/Plugins/Experimental/PCG/Source/PCG/Private/Elements/PCGDifferenceElement.cpp b/Engine/Plugins/Experimental/PCG/Source/PCG/Private/Elements/PCGDifferenceElement.cpp
--- a/Engine/Plugins/Experimental/PCG/Source/PCG/Private/Elements/PCGDifferenceElement.cpp
+++ b/Engine/Plugins/Experimental/PCG/Source/PCG/Private/Elements/PCGDifferenceElement.cpp
@@ -7,6 +7,8 @@
 #include "Data/PCGUnionData.h"
 #include "Helpers/PCGSettingsHelpers.h"
 
+#include UE_INLINE_GENERATED_CPP_BY_NAME(PCGDifferenceElement)
+
 TArray<FPCGPinProperties> UPCGDifferenceSettings::InputPinProperties() const
 {
 	TArray<FPCGPinProperties> PinProperties;
@@ -197,4 +199,4 @@ void FPCGDifferenceElement::LabellessProcessing(FPCGContext* Context) const
 
 	// Finally, pass-through settings
 	Outputs.Append(Context->InputData.GetAllSettings());
-}
\ No newline at end of file
+}
diff --git a/Engine/Plugins/Experimental/PCG/Source/PCG/Private/Elements/PCGExecuteBlueprint.cpp b/Engine/Plugins/Experimental/PCG/Source/PCG/Private/Elements/PCGExecuteBlueprint.cpp
--- a/Engine/Plugins/Experimental/PCG/Source/PCG/Private/Elements/PCGExecuteBlueprint.cpp
+++ b/Engine/Plugins/Experimental/PCG/Source/PCG/Private/Elements/PCGExecuteBlueprint.cpp
@@ -12,6 +12,8 @@
 #include "Engine/World.h"
 #include "UObject/Package.h"
 
+#include UE_INLINE_GENERATED_CPP_BY_NAME(PCGExecuteBlueprint)
+
 #if WITH_EDITOR
 namespace PCGBlueprintHelper
 {
@@ -716,4 +718,4 @@ bool FPCGExecuteBlueprintElement::CanExecuteOnlyOnMainThread(FPCGContext* Contex
 	{
 		return false;
 	}
-}
\ No newline at end of file
+}
diff --git a/Engine/Plugins/Experimental/PCG/Source/PCG/Private/Elements/PCGIntersectionElement.cpp b/Engine/Plugins/Experimental/PCG/Source/PCG/Private/Elements/PCGIntersectionElement.cpp
--- a/Engine/Plugins/Experimental/PCG/Source/PCG/Private/Elements/PCGIntersectionElement.cpp
+++ b/Engine/Plugins/Experimental/PCG/Source/PCG/Private/Elements/PCGIntersectionElement.cpp
@@ -3,6 +3,8 @@
 #include "Elements/PCGIntersectionElement.h"
 #include "Helpers/PCGSettingsHelpers.h"
 
+#include UE_INLINE_GENERATED_CPP_BY_NAME(PCGIntersectionElement)
+
 FPCGElementPtr UPCGIntersectionSettings::CreateElement() const
 {
 	return MakeShared<FPCGIntersectionElement>();
@@ -68,4 +70,4 @@ bool FPCGIntersectionElement::ExecuteInternal(FPCGContext* Context) const
 	Outputs.Append(Context->InputData.GetAllSettings());
 
 	return true;
-}
\ No newline at end of file
+}
diff --git a/Engine/Plugins/Experimental/PCG/Source/PCG/Private/Elements/PCGPointExtentsModifier.cpp b/Engine/Plugins/Experimental/PCG/Source/PCG/Private/Elements/PCGPointExtentsModifier.cpp
--- a/Engine/Plugins/Experimental/PCG/Source/PCG/Private/Elements/PCGPointExtentsModifier.cpp
+++ b/Engine/Plugins/Experimental/PCG/Source/PCG/Private/Elements/PCGPointExtentsModifier.cpp
@@ -5,6 +5,8 @@
 #include "PCGHelpers.h"
 #include "Helpers/PCGSettingsHelpers.h"
 
+#include UE_INLINE_GENERATED_CPP_BY_NAME(PCGPointExtentsModifier)
+
 FPCGElementPtr UPCGPointExtentsModifierSettings::CreateElement() const
 {
 	return MakeShared<FPCGPointExtentsModifier>();
@@ -86,3 +88,4 @@ bool FPCGPointExtentsModifier::ExecuteInternal(FPCGContext* Context) const
 	
 	return true;
 }
+
diff --git a/Engine/Plugins/Experimental/PCG/Source/PCG/Private/Elements/PCGPointFilter.cpp b/Engine/Plugins/Experimental/PCG/Source/PCG/Private/Elements/PCGPointFilter.cpp
--- a/Engine/Plugins/Experimental/PCG/Source/PCG/Private/Elements/PCGPointFilter.cpp
+++ b/Engine/Plugins/Experimental/PCG/Source/PCG/Private/Elements/PCGPointFilter.cpp
@@ -9,6 +9,8 @@
 #include "Metadata/PCGMetadataAttributeTpl.h"
 #include "Helpers/PCGAsync.h"
 
+#include UE_INLINE_GENERATED_CPP_BY_NAME(PCGPointFilter)
+
 namespace PCGPointFilterConstants
 {
 	const FName DataToFilterLabel = TEXT("In");
@@ -860,4 +862,4 @@ bool FPCGPointFilterElement::ExecuteInternal(FPCGContext* Context) const
 	}
 
 	return true;
-}
\ No newline at end of file
+}
diff --git a/Engine/Plugins/Experimental/PCG/Source/PCG/Private/Elements/PCGPointSampler.cpp b/Engine/Plugins/Experimental/PCG/Source/PCG/Private/Elements/PCGPointSampler.cpp
--- a/Engine/Plugins/Experimental/PCG/Source/PCG/Private/Elements/PCGPointSampler.cpp
+++ b/Engine/Plugins/Experimental/PCG/Source/PCG/Private/Elements/PCGPointSampler.cpp
@@ -7,6 +7,8 @@
 #include "Helpers/PCGSettingsHelpers.h"
 #include "Math/RandomStream.h"
 
+#include UE_INLINE_GENERATED_CPP_BY_NAME(PCGPointSampler)
+
 UPCGPointSamplerSettings::UPCGPointSamplerSettings()
 {
 	bUseSeed = true;
@@ -136,3 +138,4 @@ bool FPCGPointSamplerElement::ExecuteInternal(FPCGContext* Context) const
 
 	return true;
 }
+
diff --git a/Engine/Plugins/Experimental/PCG/Source/PCG/Private/Elements/PCGProjectionElement.cpp b/Engine/Plugins/Experimental/PCG/Source/PCG/Private/Elements/PCGProjectionElement.cpp
--- a/Engine/Plugins/Experimental/PCG/Source/PCG/Private/Elements/PCGProjectionElement.cpp
+++ b/Engine/Plugins/Experimental/PCG/Source/PCG/Private/Elements/PCGProjectionElement.cpp
@@ -6,6 +6,8 @@
 #include "PCGEdge.h"
 #include "Helpers/PCGSettingsHelpers.h"
 
+#include UE_INLINE_GENERATED_CPP_BY_NAME(PCGProjectionElement)
+
 #define LOCTEXT_NAMESPACE "PCGProjectionElement"
 
 TArray<FPCGPinProperties> UPCGProjectionSettings::InputPinProperties() const
@@ -142,3 +144,4 @@ void UPCGProjectionSettings::ApplyDeprecation(UPCGNode* InOutNode)
 #endif // WITH_EDITOR
 
 #undef LOCTEXT_NAMESPACE
+
diff --git a/Engine/Plugins/Experimental/PCG/Source/PCG/Private/Elements/PCGPropertyToParamData.cpp b/Engine/Plugins/Experimental/PCG/Source/PCG/Private/Elements/PCGPropertyToParamData.cpp
--- a/Engine/Plugins/Experimental/PCG/Source/PCG/Private/Elements/PCGPropertyToParamData.cpp
+++ b/Engine/Plugins/Experimental/PCG/Source/PCG/Private/Elements/PCGPropertyToParamData.cpp
@@ -12,6 +12,8 @@
 
 #include "EngineUtils.h"
 
+#include UE_INLINE_GENERATED_CPP_BY_NAME(PCGPropertyToParamData)
+
 void UPCGPropertyToParamDataSettings::PostLoad()
 {
 	Super::PostLoad();
@@ -123,4 +125,4 @@ bool FPCGPropertyToParamDataElement::ExecuteInternal(FPCGContext* Context) const
 	Output.Data = ParamData;
 
 	return true;
-}
\ No newline at end of file
+}
diff --git a/Engine/Plugins/Experimental/PCG/Source/PCG/Private/Elements/PCGSelfPruning.cpp b/Engine/Plugins/Experimental/PCG/Source/PCG/Private/Elements/PCGSelfPruning.cpp
--- a/Engine/Plugins/Experimental/PCG/Source/PCG/Private/Elements/PCGSelfPruning.cpp
+++ b/Engine/Plugins/Experimental/PCG/Source/PCG/Private/Elements/PCGSelfPruning.cpp
@@ -7,6 +7,8 @@
 #include "PCGHelpers.h"
 #include "Math/RandomStream.h"
 
+#include UE_INLINE_GENERATED_CPP_BY_NAME(PCGSelfPruning)
+
 namespace PCGSelfPruningAlgorithms
 {
 	bool RandomSort(const FPCGPointRef& A, const FPCGPointRef& B)
@@ -261,4 +263,4 @@ bool FPCGSelfPruningElement::ExecuteInternal(FPCGContext* Context) const
 	PCGSelfPruningElement::Execute(Context, PruningType, RadiusSimilarityFactor, bRandomizedPruning);
 
 	return true;
-}
\ No newline at end of file
+}
diff --git a/Engine/Plugins/Experimental/PCG/Source/PCG/Private/Elements/PCGSpawnActor.cpp b/Engine/Plugins/Experimental/PCG/Source/PCG/Private/Elements/PCGSpawnActor.cpp
--- a/Engine/Plugins/Experimental/PCG/Source/PCG/Private/Elements/PCGSpawnActor.cpp
+++ b/Engine/Plugins/Experimental/PCG/Source/PCG/Private/Elements/PCGSpawnActor.cpp
@@ -17,6 +17,8 @@
 #include "Engine/World.h"
 #include "ISMPartition/ISMComponentDescriptor.h"
 
+#include UE_INLINE_GENERATED_CPP_BY_NAME(PCGSpawnActor)
+
 UPCGNode* UPCGSpawnActorSettings::CreateNode() const
 {
 	return NewObject<UPCGSpawnActorNode>();
diff --git a/Engine/Plugins/Experimental/PCG/Source/PCG/Private/Elements/PCGSplineSampler.cpp b/Engine/Plugins/Experimental/PCG/Source/PCG/Private/Elements/PCGSplineSampler.cpp
--- a/Engine/Plugins/Experimental/PCG/Source/PCG/Private/Elements/PCGSplineSampler.cpp
+++ b/Engine/Plugins/Experimental/PCG/Source/PCG/Private/Elements/PCGSplineSampler.cpp
@@ -13,6 +13,8 @@
 #include "Components/SplineComponent.h"
 #include "Voronoi/Voronoi.h"
 
+#include UE_INLINE_GENERATED_CPP_BY_NAME(PCGSplineSampler)
+
 namespace PCGSplineSamplerHelpers
 {
 	/**
@@ -951,3 +953,4 @@ bool FPCGSplineSamplerElement::ExecuteInternal(FPCGContext* Context) const
 
 	return true;
 }
+
diff --git a/Engine/Plugins/Experimental/PCG/Source/PCG/Private/Elements/PCGStaticMeshSpawner.cpp b/Engine/Plugins/Experimental/PCG/Source/PCG/Private/Elements/PCGStaticMeshSpawner.cpp
--- a/Engine/Plugins/Experimental/PCG/Source/PCG/Private/Elements/PCGStaticMeshSpawner.cpp
+++ b/Engine/Plugins/Experimental/PCG/Source/PCG/Private/Elements/PCGStaticMeshSpawner.cpp
@@ -15,6 +15,8 @@
 #include "Components/InstancedStaticMeshComponent.h"
 #include "Engine/StaticMesh.h"
 
+#include UE_INLINE_GENERATED_CPP_BY_NAME(PCGStaticMeshSpawner)
+
 UPCGStaticMeshSpawnerSettings::UPCGStaticMeshSpawnerSettings(const FObjectInitializer &ObjectInitializer)
 {
 	bUseSeed = true;
@@ -333,3 +335,4 @@ void UPCGStaticMeshSpawnerSettings::RefreshInstancePacker()
 		InstancePackerInstance = nullptr;
 	}
 }
+
diff --git a/Engine/Plugins/Experimental/PCG/Source/PCG/Private/Elements/PCGSurfaceSampler.cpp b/Engine/Plugins/Experimental/PCG/Source/PCG/Private/Elements/PCGSurfaceSampler.cpp
--- a/Engine/Plugins/Experimental/PCG/Source/PCG/Private/Elements/PCGSurfaceSampler.cpp
+++ b/Engine/Plugins/Experimental/PCG/Source/PCG/Private/Elements/PCGSurfaceSampler.cpp
@@ -11,6 +11,8 @@
 
 #include "Math/RandomStream.h"
 
+#include UE_INLINE_GENERATED_CPP_BY_NAME(PCGSurfaceSampler)
+
 namespace PCGSurfaceSampler
 {
 	bool FSurfaceSamplerSettings::Initialize(const UPCGSurfaceSamplerSettings* InSettings, FPCGContext* Context, const FBox& InputBounds)
@@ -241,4 +243,4 @@ bool FPCGSurfaceSamplerElement::ExecuteInternal(FPCGContext* Context) const
 	Outputs.Append(Context->InputData.GetAllSettings());
 
 	return true;
-}
\ No newline at end of file
+}
diff --git a/Engine/Plugins/Experimental/PCG/Source/PCG/Private/Elements/PCGTextureSampler.cpp b/Engine/Plugins/Experimental/PCG/Source/PCG/Private/Elements/PCGTextureSampler.cpp
--- a/Engine/Plugins/Experimental/PCG/Source/PCG/Private/Elements/PCGTextureSampler.cpp
+++ b/Engine/Plugins/Experimental/PCG/Source/PCG/Private/Elements/PCGTextureSampler.cpp
@@ -9,6 +9,8 @@
 #include "Helpers/PCGBlueprintHelpers.h"
 #include "Helpers/PCGSettingsHelpers.h"
 
+#include UE_INLINE_GENERATED_CPP_BY_NAME(PCGTextureSampler)
+
 FPCGElementPtr UPCGTextureSamplerSettings::CreateElement() const
 {
 	return MakeShared<FPCGTextureSamplerElement>();
@@ -88,4 +90,4 @@ bool FPCGTextureSamplerElement::ExecuteInternal(FPCGContext* Context) const
 	TextureData->TileBounds = FBox2D(TileBoundsMin, TileBoundsMax);
 
 	return true;
-}
\ No newline at end of file
+}
diff --git a/Engine/Plugins/Experimental/PCG/Source/PCG/Private/Elements/PCGTransformPoints.cpp b/Engine/Plugins/Experimental/PCG/Source/PCG/Private/Elements/PCGTransformPoints.cpp
--- a/Engine/Plugins/Experimental/PCG/Source/PCG/Private/Elements/PCGTransformPoints.cpp
+++ b/Engine/Plugins/Experimental/PCG/Source/PCG/Private/Elements/PCGTransformPoints.cpp
@@ -5,6 +5,8 @@
 #include "PCGHelpers.h"
 #include "Helpers/PCGSettingsHelpers.h"
 
+#include UE_INLINE_GENERATED_CPP_BY_NAME(PCGTransformPoints)
+
 UPCGTransformPointsSettings::UPCGTransformPointsSettings()
 {
 	bUseSeed = true;
@@ -119,3 +121,4 @@ bool FPCGTransformPointsElement::ExecuteInternal(FPCGContext* Context) const
 
 	return true;
 }
+
diff --git a/Engine/Plugins/Experimental/PCG/Source/PCG/Private/Elements/PCGUnionElement.cpp b/Engine/Plugins/Experimental/PCG/Source/PCG/Private/Elements/PCGUnionElement.cpp
--- a/Engine/Plugins/Experimental/PCG/Source/PCG/Private/Elements/PCGUnionElement.cpp
+++ b/Engine/Plugins/Experimental/PCG/Source/PCG/Private/Elements/PCGUnionElement.cpp
@@ -3,6 +3,8 @@
 #include "Elements/PCGUnionElement.h"
 #include "Helpers/PCGSettingsHelpers.h"
 
+#include UE_INLINE_GENERATED_CPP_BY_NAME(PCGUnionElement)
+
 FPCGElementPtr UPCGUnionSettings::CreateElement() const
 {
 	return MakeShared<FPCGUnionElement>();
@@ -70,4 +72,4 @@ bool FPCGUnionElement::ExecuteInternal(FPCGContext* Context) const
 	Outputs.Append(Context->InputData.GetAllSettings());
 
 	return true;
-}
\ No newline at end of file
+}
diff --git a/Engine/Plugins/Experimental/PCG/Source/PCG/Private/Elements/PCGVolumeSampler.cpp b/Engine/Plugins/Experimental/PCG/Source/PCG/Private/Elements/PCGVolumeSampler.cpp
--- a/Engine/Plugins/Experimental/PCG/Source/PCG/Private/Elements/PCGVolumeSampler.cpp
+++ b/Engine/Plugins/Experimental/PCG/Source/PCG/Private/Elements/PCGVolumeSampler.cpp
@@ -7,6 +7,8 @@
 #include "Helpers/PCGSettingsHelpers.h"
 #include "PCGHelpers.h"
 
+#include UE_INLINE_GENERATED_CPP_BY_NAME(PCGVolumeSampler)
+
 namespace PCGVolumeSampler
 {
 	UPCGPointData* SampleVolume(FPCGContext* Context, const UPCGSpatialData* SpatialData, const FVolumeSamplerSettings& SamplerSettings)
@@ -116,4 +118,4 @@ bool FPCGVolumeSamplerElement::ExecuteInternal(FPCGContext* Context) const
 	Outputs.Append(Context->InputData.GetAllSettings());
 
 	return true;
-}
\ No newline at end of file
+}
diff --git a/Engine/Plugins/Experimental/PCG/Source/PCG/Private/Elements/PCGWorldQuery.cpp b/Engine/Plugins/Experimental/PCG/Source/PCG/Private/Elements/PCGWorldQuery.cpp
--- a/Engine/Plugins/Experimental/PCG/Source/PCG/Private/Elements/PCGWorldQuery.cpp
+++ b/Engine/Plugins/Experimental/PCG/Source/PCG/Private/Elements/PCGWorldQuery.cpp
@@ -7,6 +7,8 @@
 
 #include "GameFramework/Actor.h"
 
+#include UE_INLINE_GENERATED_CPP_BY_NAME(PCGWorldQuery)
+
 FPCGElementPtr UPCGWorldQuerySettings::CreateElement() const
 {
 	return MakeShared<FPCGWorldVolumetricQueryElement>();
@@ -108,4 +110,4 @@ bool FPCGWorldRayHitQueryElement::ExecuteInternal(FPCGContext* Context) const
 	Context->OutputData.TaggedData.Append(Context->InputData.GetAllSettings());
 
 	return true;
-}
\ No newline at end of file
+}
diff --git a/Engine/Plugins/Experimental/PCG/Source/PCG/Private/Grid/PCGLandscapeCache.cpp b/Engine/Plugins/Experimental/PCG/Source/PCG/Private/Grid/PCGLandscapeCache.cpp
--- a/Engine/Plugins/Experimental/PCG/Source/PCG/Private/Grid/PCGLandscapeCache.cpp
+++ b/Engine/Plugins/Experimental/PCG/Source/PCG/Private/Grid/PCGLandscapeCache.cpp
@@ -24,6 +24,8 @@
 #include "Serialization/BufferReader.h"
 #include "Serialization/BufferWriter.h"
 
+#include UE_INLINE_GENERATED_CPP_BY_NAME(PCGLandscapeCache)
+
 static TAutoConsoleVariable<float> CVarLandscapeCacheGCFrequency(
 	TEXT("pcg.LandscapeCacheGCFrequency"),
 	10,
@@ -728,4 +730,4 @@ void UPCGLandscapeCache::CacheLayerNames(ALandscapeProxy* Landscape)
 	}
 }
 
-#endif // WITH_EDITOR
\ No newline at end of file
+#endif // WITH_EDITOR
diff --git a/Engine/Plugins/Experimental/PCG/Source/PCG/Private/Grid/PCGPartitionActor.cpp b/Engine/Plugins/Experimental/PCG/Source/PCG/Private/Grid/PCGPartitionActor.cpp
--- a/Engine/Plugins/Experimental/PCG/Source/PCG/Private/Grid/PCGPartitionActor.cpp
+++ b/Engine/Plugins/Experimental/PCG/Source/PCG/Private/Grid/PCGPartitionActor.cpp
@@ -11,6 +11,8 @@
 #include "Landscape.h"
 #include "Components/BoxComponent.h"
 
+#include UE_INLINE_GENERATED_CPP_BY_NAME(PCGPartitionActor)
+
 constexpr uint32 InvalidPCGGridSizeValue = 0u;
 
 APCGPartitionActor::APCGPartitionActor(const FObjectInitializer& ObjectInitializer)
@@ -415,4 +417,4 @@ TSet<TObjectPtr<UPCGComponent>> APCGPartitionActor::GetAllOriginalPCGComponents(
 	return ResultComponents;
 }
 
-#endif // WITH_EDITOR
\ No newline at end of file
+#endif // WITH_EDITOR
diff --git a/Engine/Plugins/Experimental/PCG/Source/PCG/Private/Helpers/PCGActorHelpers.cpp b/Engine/Plugins/Experimental/PCG/Source/PCG/Private/Helpers/PCGActorHelpers.cpp
--- a/Engine/Plugins/Experimental/PCG/Source/PCG/Private/Helpers/PCGActorHelpers.cpp
+++ b/Engine/Plugins/Experimental/PCG/Source/PCG/Private/Helpers/PCGActorHelpers.cpp
@@ -18,6 +18,8 @@
 #include "WorldPartition/WorldPartition.h"
 #include "ISMPartition/ISMComponentDescriptor.h"
 
+#include UE_INLINE_GENERATED_CPP_BY_NAME(PCGActorHelpers)
+
 #if WITH_EDITOR
 #include "Editor.h"
 #include "SourceControlHelpers.h"
@@ -483,3 +485,4 @@ FIntVector UPCGActorHelpers::GetCellCoord(FVector InPosition, int InGridSize, bo
 		bUse2DGrid ? 0 : FMath::FloorToInt(Temp.Z)
 	);
 }
+
diff --git a/Engine/Plugins/Experimental/PCG/Source/PCG/Private/Helpers/PCGBlueprintHelpers.cpp b/Engine/Plugins/Experimental/PCG/Source/PCG/Private/Helpers/PCGBlueprintHelpers.cpp
--- a/Engine/Plugins/Experimental/PCG/Source/PCG/Private/Helpers/PCGBlueprintHelpers.cpp
+++ b/Engine/Plugins/Experimental/PCG/Source/PCG/Private/Helpers/PCGBlueprintHelpers.cpp
@@ -15,6 +15,8 @@
 #include "LandscapeProxy.h"
 #include "LandscapeInfo.h"
 
+#include UE_INLINE_GENERATED_CPP_BY_NAME(PCGBlueprintHelpers)
+
 int UPCGBlueprintHelpers::ComputeSeedFromPosition(const FVector& InPosition)
 {
 	// TODO: should have a config to drive this
@@ -184,4 +186,4 @@ TArray<FPCGLandscapeLayerWeight> UPCGBlueprintHelpers::GetInterpolatedPCGLandsca
 	});
 					
 	return Result;
-}
\ No newline at end of file
+}
diff --git a/Engine/Plugins/Experimental/PCG/Source/PCG/Private/InstancePackers/PCGInstancePackerBase.cpp b/Engine/Plugins/Experimental/PCG/Source/PCG/Private/InstancePackers/PCGInstancePackerBase.cpp
--- a/Engine/Plugins/Experimental/PCG/Source/PCG/Private/InstancePackers/PCGInstancePackerBase.cpp
+++ b/Engine/Plugins/Experimental/PCG/Source/PCG/Private/InstancePackers/PCGInstancePackerBase.cpp
@@ -7,6 +7,8 @@
 #include "Metadata/PCGMetadataAttributeTraits.h"
 #include "Metadata/PCGMetadataAttributeTpl.h"
 
+#include UE_INLINE_GENERATED_CPP_BY_NAME(PCGInstancePackerBase)
+
 bool UPCGInstancePackerBase::AddTypeToPacking(int TypeId, FPCGPackedCustomData& OutPackedCustomData) const
 {
 	if (TypeId == PCG::Private::MetadataTypes<float>::Id)
@@ -137,3 +139,4 @@ void UPCGInstancePackerBase::PackCustomDataFromAttributes(const FPCGMeshInstance
 	}
 }
 
+
diff --git a/Engine/Plugins/Experimental/PCG/Source/PCG/Private/InstancePackers/PCGInstancePackerByAttribute.cpp b/Engine/Plugins/Experimental/PCG/Source/PCG/Private/InstancePackers/PCGInstancePackerByAttribute.cpp
--- a/Engine/Plugins/Experimental/PCG/Source/PCG/Private/InstancePackers/PCGInstancePackerByAttribute.cpp
+++ b/Engine/Plugins/Experimental/PCG/Source/PCG/Private/InstancePackers/PCGInstancePackerByAttribute.cpp
@@ -3,6 +3,8 @@
 #include "InstancePackers/PCGInstancePackerByAttribute.h"
 #include "Metadata/PCGMetadataAttribute.h"
 
+#include UE_INLINE_GENERATED_CPP_BY_NAME(PCGInstancePackerByAttribute)
+
 void UPCGInstancePackerByAttribute::PackInstances_Implementation(FPCGContext& Context, const UPCGSpatialData* InSpatialData, const FPCGMeshInstanceList& InstanceList, FPCGPackedCustomData& OutPackedCustomData) const
 {
 	if (!InSpatialData || !InSpatialData->Metadata)
@@ -36,3 +38,4 @@ void UPCGInstancePackerByAttribute::PackInstances_Implementation(FPCGContext& Co
 
 	PackCustomDataFromAttributes(InstanceList, SelectedAttributes, OutPackedCustomData);
  }
+
