Add flag to enable disable playback on individual component in a chaos cache manager

diff --git a/Engine/Plugins/Experimental/ChaosCaching/Source/ChaosCaching/Private/Chaos/CacheManagerActor.cpp b/Engine/Plugins/Experimental/ChaosCaching/Source/ChaosCaching/Private/Chaos/CacheManagerActor.cpp
--- a/Engine/Plugins/Experimental/ChaosCaching/Source/ChaosCaching/Private/Chaos/CacheManagerActor.cpp
+++ b/Engine/Plugins/Experimental/ChaosCaching/Source/ChaosCaching/Private/Chaos/CacheManagerActor.cpp
@@ -41,6 +41,11 @@ void FObservedComponent::ResetRuntimeData(const EStartMode ManagerStartMode)
 	TickRecord.Reset();
 }
 
+bool FObservedComponent::IsEnabled(ECacheMode CacheMode) const
+{
+	return (CacheMode == ECacheMode::Record) || bPlaybackEnabled;
+}
+
 void FObservedComponent::PostSerialize(const FArchive& Ar)
 {
 	PRAGMA_DISABLE_DEPRECATION_WARNINGS
@@ -319,7 +324,13 @@ void AChaosCacheManager::BeginEvaluate()
 
 		// Reset timers and last cache
 		Observed.ResetRuntimeData(StartMode);
-		
+
+		// we need to create the adapters regardless of it being enabled for playback and we can bail if necessary after that 
+		if (!Observed.IsEnabled(CacheMode))
+		{
+			continue;
+		}
+
 		bool                    bRequiresRecord = false;
 		FComponentCacheAdapter* CurrAdapter     = ActiveAdapters[Index];
 		check(CurrAdapter);    // should definitely have added one above
@@ -596,7 +607,7 @@ void AChaosCacheManager::OnStartFrameChanged(Chaos::FReal InTime)
 				Observed.Cache = CacheCollection->FindCache(Observed.CacheName);
 			}
 
-			if (!Observed.Cache || !Observed.BestFitAdapter || Observed.Cache->GetDuration()==0.0)
+			if (!Observed.Cache || !Observed.BestFitAdapter || Observed.Cache->GetDuration()==0.0 || Observed.IsEnabled(CacheMode))
 			{
 				continue;
 			}
@@ -617,8 +628,8 @@ void AChaosCacheManager::OnStartFrameChanged(Chaos::FReal InTime)
 void AChaosCacheManager::TriggerComponent(UPrimitiveComponent* InComponent)
 {
 	// #BGTODO Maybe not totally thread-safe, probably safer with an atomic or condition var rather than the bTriggered flag
-	FObservedComponent* Found = Algo::FindByPredicate(ObservedComponents, [InComponent](const FObservedComponent& Test) {
-		return Test.GetComponent() == InComponent;
+	FObservedComponent* Found = Algo::FindByPredicate(ObservedComponents, [this, InComponent](const FObservedComponent& Test) {
+		return Test.GetComponent() == InComponent && Test.IsEnabled(CacheMode);
 	});
 
 	if (Found && StartMode == EStartMode::Triggered)
@@ -629,8 +640,8 @@ void AChaosCacheManager::TriggerComponent(UPrimitiveComponent* InComponent)
 
 void AChaosCacheManager::TriggerComponentByCache(FName InCacheName)
 {
-	FObservedComponent* Found = Algo::FindByPredicate(ObservedComponents, [InCacheName](const FObservedComponent& Test) {
-		return Test.CacheName == InCacheName && Test.GetComponent();
+	FObservedComponent* Found = Algo::FindByPredicate(ObservedComponents, [this, InCacheName](const FObservedComponent& Test) {
+		return Test.CacheName == InCacheName && Test.GetComponent() && Test.IsEnabled(CacheMode);
 	});
 
 	if (Found && StartMode == EStartMode::Triggered)
@@ -643,7 +654,7 @@ void AChaosCacheManager::TriggerAll()
 {
 	for(FObservedComponent& Observed : ObservedComponents)
 	{
-		if (StartMode == EStartMode::Triggered && Observed.GetComponent())
+		if (StartMode == EStartMode::Triggered && Observed.GetComponent() && Observed.IsEnabled(CacheMode))
 		{
 			Observed.bTriggered = true;
 		}
@@ -701,7 +712,7 @@ void AChaosCacheManager::TickObservedComponents(const TArray<int32>& InIndices,
 		FObservedComponent&            Observed = ObservedComponents[Index];
 		Chaos::FComponentCacheAdapter* Adapter  = ActiveAdapters[Index];
 
-		if(!Observed.Cache)
+		if(!Observed.Cache || !Observed.IsEnabled(CacheMode))
 		{
 			// Skip if no available cache - this can happen if a component was deleted while being observed - the other components
 			// can play fine, we just omit any that we cannot find.
diff --git a/Engine/Plugins/Experimental/ChaosCaching/Source/ChaosCaching/Public/Chaos/CacheManagerActor.h b/Engine/Plugins/Experimental/ChaosCaching/Source/ChaosCaching/Public/Chaos/CacheManagerActor.h
--- a/Engine/Plugins/Experimental/ChaosCaching/Source/ChaosCaching/Public/Chaos/CacheManagerActor.h
+++ b/Engine/Plugins/Experimental/ChaosCaching/Source/ChaosCaching/Public/Chaos/CacheManagerActor.h
@@ -48,6 +48,7 @@ struct CHAOSCACHING_API FObservedComponent
 	FObservedComponent()
 		: CacheName(NAME_None)
 		, bIsSimulating(true)
+		, bPlaybackEnabled(true)
 		, Cache(nullptr)
 		, BestFitAdapter(nullptr)
 	{
@@ -57,6 +58,7 @@ struct CHAOSCACHING_API FObservedComponent
 	{
 		CacheName = OtherComponent.CacheName;
 		bIsSimulating = OtherComponent.bIsSimulating;
+		bPlaybackEnabled = OtherComponent.bPlaybackEnabled;
 		bTriggered = OtherComponent.bTriggered;
 		AbsoluteTime = OtherComponent.AbsoluteTime;
 		TimeSinceTrigger = OtherComponent.TimeSinceTrigger;
@@ -70,6 +72,7 @@ struct CHAOSCACHING_API FObservedComponent
 	{
 		CacheName = OtherComponent.CacheName;
 		bIsSimulating = OtherComponent.bIsSimulating;
+		bPlaybackEnabled = OtherComponent.bPlaybackEnabled;
 		bTriggered = OtherComponent.bTriggered;
 		AbsoluteTime = OtherComponent.AbsoluteTime;
 		TimeSinceTrigger = OtherComponent.TimeSinceTrigger;
@@ -82,7 +85,7 @@ struct CHAOSCACHING_API FObservedComponent
 	}
 
 	/** Unique name for the cache, used as a key into the cache collection */
-	UPROPERTY(EditAnywhere, Category = "Caching")
+	UPROPERTY(EditAnywhere, BlueprintReadWrite, Category = "Caching")
 	FName CacheName;
 
 	/** Deprecated hard object reference. Not working with sequencer and take recorder since
@@ -96,7 +99,7 @@ struct CHAOSCACHING_API FObservedComponent
 	FSoftComponentReference SoftComponentRef;
 
 	/** Capture of the initial state of the component before cache manager takes control. */
-	UPROPERTY(EditAnywhere, Category = "Caching")
+	UPROPERTY(EditAnywhere, BlueprintReadOnly, Category = "Caching")
 	bool bIsSimulating;
 
 	/** Post serialize function to transfer datas from the deprecated TObjectPtr -> TSoftObjectPtr */
@@ -105,10 +108,18 @@ struct CHAOSCACHING_API FObservedComponent
 	/** Prepare runtime tick data for a new run */
 	void ResetRuntimeData(const EStartMode ManagerStartMode);
 
+	/** Check if the Observed component is enabled for a specific cache mode */
+	bool IsEnabled(ECacheMode CacheMode) const;
+
 	/** Gets the component from the internal component ref */
 	UPrimitiveComponent* GetComponent();
 	UPrimitiveComponent* GetComponent() const;
 
+private:
+	/** Whether this component is enabled for playback, this allow a cache to hold many component but only replay some of them. */
+	UPROPERTY(EditAnywhere, Category = "Caching")
+	bool bPlaybackEnabled;
+
 private:
 	friend class AChaosCacheManager;
 
