Chaos : Fix cache
- not properly recording if break event is not set on the geometry collection
- wrongly skipping enabled observed components

diff --git a/Engine/Plugins/Experimental/ChaosCaching/Source/ChaosCaching/Private/Chaos/CacheManagerActor.cpp b/Engine/Plugins/Experimental/ChaosCaching/Source/ChaosCaching/Private/Chaos/CacheManagerActor.cpp
--- a/Engine/Plugins/Experimental/ChaosCaching/Source/ChaosCaching/Private/Chaos/CacheManagerActor.cpp
+++ b/Engine/Plugins/Experimental/ChaosCaching/Source/ChaosCaching/Private/Chaos/CacheManagerActor.cpp
@@ -8,6 +8,7 @@
 #include "Chaos/ChaosCache.h"
 #include "ChaosSolversModule.h"
 #include "Components/BillboardComponent.h"
+#include "GeometryCollection/GeometryCollectionComponent.h"
 #include "Engine/Texture2D.h"
 #include "PBDRigidsSolver.h"
 #include "Features/IModularFeatures.h"
@@ -607,7 +608,7 @@ void AChaosCacheManager::OnStartFrameChanged(Chaos::FReal InTime)
 				Observed.Cache = CacheCollection->FindCache(Observed.CacheName);
 			}
 
-			if (!Observed.Cache || !Observed.BestFitAdapter || Observed.Cache->GetDuration()==0.0 || Observed.IsEnabled(CacheMode))
+			if (!Observed.Cache || !Observed.BestFitAdapter || Observed.Cache->GetDuration()==0.0 || !Observed.IsEnabled(CacheMode))
 			{
 				continue;
 			}
@@ -684,6 +685,14 @@ FObservedComponent& AChaosCacheManager::AddNewObservedComponent(UPrimitiveCompon
 
 	NewEntry.CacheName = MakeUniqueObjectName(CacheCollection, UChaosCache::StaticClass(), CacheName);
 
+	// make sure we keep track of the various flag that may be changed by the cahe to function
+	NewEntry.bIsSimulating = InComponent->IsSimulatingPhysics();
+	NewEntry.bHasNotifyBreaks = false;
+	if (UGeometryCollectionComponent* GeomComponent = Cast<UGeometryCollectionComponent>(InComponent))
+	{
+		NewEntry.bHasNotifyBreaks = GeomComponent->bNotifyBreaks;
+	}
+
 	return NewEntry;
 }
 
@@ -745,11 +754,21 @@ void AChaosCacheManager::SetObservedComponentProperties(const ECacheMode& NewCac
 		{
 			if (NewCacheMode == ECacheMode::Record)
 			{
-				PrimComp->BodyInstance.bSimulatePhysics = ObservedComponent.bIsSimulating;	
+				PrimComp->BodyInstance.bSimulatePhysics = ObservedComponent.bIsSimulating;
+				if (UGeometryCollectionComponent* GeomComponent = Cast<UGeometryCollectionComponent>(PrimComp))
+				{
+					// in record mode we need to have notify break on on the proxy 
+					GeomComponent->SetNotifyBreaks(true);
+				}
 			}
 			else
 			{
 				PrimComp->BodyInstance.bSimulatePhysics = false;
+				if (UGeometryCollectionComponent* GeomComponent = Cast<UGeometryCollectionComponent>(PrimComp))
+				{
+					// in playback modes we can restore the state of the notification 
+					GeomComponent->SetNotifyBreaks(ObservedComponent.bHasNotifyBreaks);
+				}
 			}
 		}
 	}
diff --git a/Engine/Plugins/Experimental/ChaosCaching/Source/ChaosCaching/Public/Chaos/CacheManagerActor.h b/Engine/Plugins/Experimental/ChaosCaching/Source/ChaosCaching/Public/Chaos/CacheManagerActor.h
--- a/Engine/Plugins/Experimental/ChaosCaching/Source/ChaosCaching/Public/Chaos/CacheManagerActor.h
+++ b/Engine/Plugins/Experimental/ChaosCaching/Source/ChaosCaching/Public/Chaos/CacheManagerActor.h
@@ -47,6 +47,7 @@ struct CHAOSCACHING_API FObservedComponent
 	FObservedComponent()
 		: CacheName(NAME_None)
 		, bIsSimulating(true)
+		, bHasNotifyBreaks(false)
 		, bPlaybackEnabled(true)
 		, Cache(nullptr)
 		, BestFitAdapter(nullptr)
@@ -57,6 +58,7 @@ struct CHAOSCACHING_API FObservedComponent
 	{
 		CacheName = OtherComponent.CacheName;
 		bIsSimulating = OtherComponent.bIsSimulating;
+		bHasNotifyBreaks = OtherComponent.bHasNotifyBreaks;
 		bPlaybackEnabled = OtherComponent.bPlaybackEnabled;
 		bTriggered = OtherComponent.bTriggered;
 		AbsoluteTime = OtherComponent.AbsoluteTime;
@@ -71,6 +73,7 @@ struct CHAOSCACHING_API FObservedComponent
 	{
 		CacheName = OtherComponent.CacheName;
 		bIsSimulating = OtherComponent.bIsSimulating;
+		bHasNotifyBreaks = OtherComponent.bHasNotifyBreaks;
 		bPlaybackEnabled = OtherComponent.bPlaybackEnabled;
 		bTriggered = OtherComponent.bTriggered;
 		AbsoluteTime = OtherComponent.AbsoluteTime;
@@ -101,6 +104,13 @@ struct CHAOSCACHING_API FObservedComponent
 	UPROPERTY(EditAnywhere, BlueprintReadOnly, Category = "Caching")
 	bool bIsSimulating;
 
+	/** 
+	* Capture the state of bNotifyBreaks of the component before cache manager takes control. 
+	* this is because when recording the cache needs the component to have bNotifyBreaks set on the component 
+	* to be able to properly record when clusters are breaking into smaller pieces
+	*/
+	bool bHasNotifyBreaks;
+
 	/** Post serialize function to transfer datas from the deprecated TObjectPtr -> TSoftObjectPtr */
 	void PostSerialize(const FArchive& Ar);
 
