Chaos caching : allow setting EnablePlayback  and CacheCollection from blueprint

diff --git a/Engine/Plugins/Experimental/ChaosCaching/Source/ChaosCaching/Private/Chaos/CacheManagerActor.cpp b/Engine/Plugins/Experimental/ChaosCaching/Source/ChaosCaching/Private/Chaos/CacheManagerActor.cpp
--- a/Engine/Plugins/Experimental/ChaosCaching/Source/ChaosCaching/Private/Chaos/CacheManagerActor.cpp
+++ b/Engine/Plugins/Experimental/ChaosCaching/Source/ChaosCaching/Private/Chaos/CacheManagerActor.cpp
@@ -626,6 +626,21 @@ void AChaosCacheManager::OnStartFrameChanged(Chaos::FReal InTime)
 	}
 }
 
+void AChaosCacheManager::SetCacheCollection(UChaosCacheCollection* InCacheCollection)
+{
+	// check if we have any observed component already triggered
+	for (FObservedComponent& Observed : ObservedComponents)
+	{
+		if (Observed.bTriggered)
+		{
+			UE_LOG(LogChaosCache, Warning, TEXT("Trying to change the cache collection while some observed component are already been triggered ( %s )"), *GetName());
+			return;
+		}
+	}
+
+	CacheCollection = InCacheCollection;
+}
+
 void AChaosCacheManager::TriggerComponent(UPrimitiveComponent* InComponent)
 {
 	// #BGTODO Maybe not totally thread-safe, probably safer with an atomic or condition var rather than the bTriggered flag
@@ -662,6 +677,27 @@ void AChaosCacheManager::TriggerAll()
 	}
 }
 
+void AChaosCacheManager::EnablePlaybackByCache(FName InCacheName, bool bEnable)
+{
+	FObservedComponent* Found = Algo::FindByPredicate(ObservedComponents, [this, InCacheName](const FObservedComponent& Test) {
+		return Test.CacheName == InCacheName;
+		});
+
+	if (Found)
+	{
+		Found->bPlaybackEnabled = bEnable;
+	}
+}
+
+
+void AChaosCacheManager::EnablePlayback(int32 Index, bool bEnable)
+{
+	if (ObservedComponents.IsValidIndex(Index))
+	{
+		ObservedComponents[Index].bPlaybackEnabled = bEnable;
+	}
+}
+
 FObservedComponent* AChaosCacheManager::FindObservedComponent(UPrimitiveComponent* InComponent)
 {
 	return Algo::FindByPredicate(ObservedComponents, [ToTest = InComponent](const FObservedComponent& Item) {
diff --git a/Engine/Plugins/Experimental/ChaosCaching/Source/ChaosCaching/Public/Chaos/CacheManagerActor.h b/Engine/Plugins/Experimental/ChaosCaching/Source/ChaosCaching/Public/Chaos/CacheManagerActor.h
--- a/Engine/Plugins/Experimental/ChaosCaching/Source/ChaosCaching/Public/Chaos/CacheManagerActor.h
+++ b/Engine/Plugins/Experimental/ChaosCaching/Source/ChaosCaching/Public/Chaos/CacheManagerActor.h
@@ -47,8 +47,8 @@ struct CHAOSCACHING_API FObservedComponent
 	FObservedComponent()
 		: CacheName(NAME_None)
 		, bIsSimulating(true)
-		, bHasNotifyBreaks(false)
 		, bPlaybackEnabled(true)
+		, bHasNotifyBreaks(false)
 		, Cache(nullptr)
 		, BestFitAdapter(nullptr)
 	{
@@ -104,6 +104,10 @@ struct CHAOSCACHING_API FObservedComponent
 	UPROPERTY(EditAnywhere, BlueprintReadOnly, Category = "Caching")
 	bool bIsSimulating;
 
+	/** Whether this component is enabled for playback, this allow a cache to hold many component but only replay some of them. */
+	UPROPERTY(EditAnywhere, BlueprintReadWrite, Category = "Caching")
+	bool bPlaybackEnabled;
+
 	/** 
 	* Capture the state of bNotifyBreaks of the component before cache manager takes control. 
 	* this is because when recording the cache needs the component to have bNotifyBreaks set on the component 
@@ -124,11 +128,6 @@ struct CHAOSCACHING_API FObservedComponent
 	UPrimitiveComponent* GetComponent();
 	UPrimitiveComponent* GetComponent() const;
 
-private:
-	/** Whether this component is enabled for playback, this allow a cache to hold many component but only replay some of them. */
-	UPROPERTY(EditAnywhere, Category = "Caching")
-	bool bPlaybackEnabled;
-
 private:
 	friend class AChaosCacheManager;
 
@@ -272,6 +271,13 @@ protected:
 	/** Evaluates and sets state for all observed components at the specified time. */
 	void OnStartFrameChanged(Chaos::FReal InT);
 
+	/**
+	* change the cache collection for this player 
+	* if the cache is playing or recording this will have no effect
+	*/
+	UFUNCTION(BlueprintCallable, Category = "Caching")
+	void SetCacheCollection(UChaosCacheCollection* InCacheCollection);
+
 	/**
 	 * Triggers a component to play or record.
 	 * If the cache manager has an observed component entry for InComponent and it is a triggered entry
@@ -294,6 +300,14 @@ protected:
 	UFUNCTION(BlueprintCallable, Category = "Caching")
 	void TriggerAll();
 
+	/** Enable playback for a specific component using its cache name */
+	UFUNCTION(BlueprintCallable, Category = "Caching")
+	void EnablePlaybackByCache(FName InCacheName, bool bEnable);
+
+	/** Enable playback for a specific component using its index in the list of observed component */
+	UFUNCTION(BlueprintCallable, Category = "Caching")
+	void EnablePlayback(int32 Index, bool bEnable);
+
 	FObservedComponent* FindObservedComponent(UPrimitiveComponent* InComponent);
 	FObservedComponent& AddNewObservedComponent(UPrimitiveComponent* InComponent);
 	FObservedComponent& FindOrAddObservedComponent(UPrimitiveComponent* InComponent);
