Fix some Chaos Cloth caching bugs related to running Take Recorder while PIE is running.

Two main issues we were running into:

* Cloth Solver getting deleted but not cleaning up bookkeeping correctly when starting Take Recorder causing undefined behavior (crashes).
* Bookkeeping and events gettinig edited when starting Take Recorder while PIE cloth sims are still in flight causing race conditions (crashes).

Fixes:
* FPhysicsSolverEvents now broadcasts a Teardown event in its destructor.
* CacheManager calls EndEvaluate at the beginning of BeginEvaluate. Double calling EndEvaluate is currently OK. Double calling BeginEvaluate without an End is not safe (and what was happening when starting Take Recorder while PIE was running), hence this change.
* Add new WaitForSolverTasks virtual function to CacheAdapter which stalls on cloth sim tasks in the SkeletalMeshComponentCacheAdapter. This is called in EndEvaluate and fixes the race conditions I was seeing.

diff --git a/Engine/Plugins/Experimental/ChaosCaching/Source/ChaosCaching/Private/Chaos/CacheManagerActor.cpp b/Engine/Plugins/Experimental/ChaosCaching/Source/ChaosCaching/Private/Chaos/CacheManagerActor.cpp
--- a/Engine/Plugins/Experimental/ChaosCaching/Source/ChaosCaching/Private/Chaos/CacheManagerActor.cpp
+++ b/Engine/Plugins/Experimental/ChaosCaching/Source/ChaosCaching/Private/Chaos/CacheManagerActor.cpp
@@ -247,6 +247,8 @@ void AChaosCacheManager::BeginPlay()
 
 void AChaosCacheManager::BeginEvaluate()
 {
+	EndEvaluate();
+
 	using namespace Chaos;
 	bIsSimulating = !(CacheMode == ECacheMode::None);
 	
@@ -271,7 +273,7 @@ void AChaosCacheManager::BeginEvaluate()
 	{
 		Adapter->Initialize();
 	}
-	ActiveAdapters.Reset();
+	check(ActiveAdapters.IsEmpty()); // Should be empty from calling EndEvaluate
 
 	int32       NumFailedPlaybackEntries = 0;
 	const int32 NumComponents            = ObservedComponents.Num();
@@ -458,6 +460,8 @@ void AChaosCacheManager::EndEvaluate()
 	FChaosSolversModule* Module = FChaosSolversModule::GetModule();
 	check(Module);
 
+	WaitForObservedComponentSolverTasks();
+
 	for(TPair<Chaos::FPhysicsSolverEvents*, FPerSolverData> PerSolver : PerSolverData)
 	{
 		Chaos::FPhysicsSolverEvents* CurrSolver = PerSolver.Key;
@@ -781,6 +785,20 @@ void AChaosCacheManager::TickObservedComponents(const TArray<int32>& InIndices,
 	}
 }
 
+void AChaosCacheManager::WaitForObservedComponentSolverTasks()
+{
+	check(ActiveAdapters.IsEmpty() || ObservedComponents.Num() == ActiveAdapters.Num());
+
+	for (int32 CompIndex = 0; CompIndex < ActiveAdapters.Num(); ++CompIndex)
+	{
+		if(const Chaos::FComponentCacheAdapter* Adapter = ActiveAdapters[CompIndex])
+		{
+			FObservedComponent& Observed = ObservedComponents[CompIndex];
+			Adapter->WaitForSolverTasks(Observed.GetComponent());
+		}
+	}
+}
+
 #if WITH_EDITOR
 void AChaosCacheManager::SetObservedComponentProperties(const ECacheMode& NewCacheMode)
 {
diff --git a/Engine/Plugins/Experimental/ChaosCaching/Source/ChaosCaching/Public/Chaos/Adapters/CacheAdapter.h b/Engine/Plugins/Experimental/ChaosCaching/Source/ChaosCaching/Public/Chaos/Adapters/CacheAdapter.h
--- a/Engine/Plugins/Experimental/ChaosCaching/Source/ChaosCaching/Public/Chaos/Adapters/CacheAdapter.h
+++ b/Engine/Plugins/Experimental/ChaosCaching/Source/ChaosCaching/Public/Chaos/Adapters/CacheAdapter.h
@@ -179,6 +179,12 @@ namespace Chaos
 		 */
 		virtual bool ValidForPlayback(UPrimitiveComponent* InComponent, UChaosCache* InCache) const = 0;
 
+		/**
+		* Stall on any solver tasks that could broadcast events via the EventsSolver. Called from game thread.
+		* This is necessary so that the cache manager can teardown and setup new event handling.
+		*/
+		virtual void WaitForSolverTasks(UPrimitiveComponent* InComponent) const {};
+
 	protected:
 	private:
 	};
diff --git a/Engine/Plugins/Experimental/ChaosCaching/Source/ChaosCaching/Public/Chaos/CacheManagerActor.h b/Engine/Plugins/Experimental/ChaosCaching/Source/ChaosCaching/Public/Chaos/CacheManagerActor.h
--- a/Engine/Plugins/Experimental/ChaosCaching/Source/ChaosCaching/Public/Chaos/CacheManagerActor.h
+++ b/Engine/Plugins/Experimental/ChaosCaching/Source/ChaosCaching/Public/Chaos/CacheManagerActor.h
@@ -337,11 +337,16 @@ private:
 	 */
 	void TickObservedComponents(const TArray<int32>& InIndices, Chaos::FReal InDt, FTickObservedFunction InCallable);
 
+	/**
+	* Stall on any in flight solver tasks that might call callbacks. Necessary before editing callbacks.
+	*/
+	void WaitForObservedComponentSolverTasks();
+
 	/** List of observed objects and their caches */
 	UPROPERTY(EditAnywhere, Category = "Caching")
 	TArray<FObservedComponent> ObservedComponents;
 
-	/** 1-1 list of adapters for the observed components, populated on BeginPlay */
+	/** 1-1 list of adapters for the observed components, populated on BeginEvaluate */
 	TArray<Chaos::FComponentCacheAdapter*> ActiveAdapters;
 
 	/** List of particles returned by the adapter as requiring a kinematic update */
diff --git a/Engine/Plugins/Experimental/ChaosCloth/Source/Chaos/Private/ChaosCloth/SkeletalMeshComponentCacheAdapter.cpp b/Engine/Plugins/Experimental/ChaosCloth/Source/Chaos/Private/ChaosCloth/SkeletalMeshComponentCacheAdapter.cpp
--- a/Engine/Plugins/Experimental/ChaosCloth/Source/Chaos/Private/ChaosCloth/SkeletalMeshComponentCacheAdapter.cpp
+++ b/Engine/Plugins/Experimental/ChaosCloth/Source/Chaos/Private/ChaosCloth/SkeletalMeshComponentCacheAdapter.cpp
@@ -1,4 +1,4 @@
-﻿// Copyright Epic Games, Inc. All Rights Reserved.
+// Copyright Epic Games, Inc. All Rights Reserved.
 
 #include "ChaosCloth/SkeletalMeshComponentCacheAdapter.h"
 #include "Chaos/ParticleHandle.h"
@@ -283,4 +283,14 @@ namespace Chaos
 		}
 		return true;
 	}
+
+	void FSkeletalMeshCacheAdapter::WaitForSolverTasks(UPrimitiveComponent* InComponent) const
+	{
+		EnsureIsInGameThreadContext();
+
+		if (USkeletalMeshComponent* MeshComp = CastChecked<USkeletalMeshComponent>(InComponent))
+		{
+			MeshComp->WaitForExistingParallelClothSimulation_GameThread();
+		}
+	}
 }    // namespace Chaos
diff --git a/Engine/Plugins/Experimental/ChaosCloth/Source/Chaos/Private/ChaosCloth/SkeletalMeshComponentCacheAdapter.h b/Engine/Plugins/Experimental/ChaosCloth/Source/Chaos/Private/ChaosCloth/SkeletalMeshComponentCacheAdapter.h
--- a/Engine/Plugins/Experimental/ChaosCloth/Source/Chaos/Private/ChaosCloth/SkeletalMeshComponentCacheAdapter.h
+++ b/Engine/Plugins/Experimental/ChaosCloth/Source/Chaos/Private/ChaosCloth/SkeletalMeshComponentCacheAdapter.h
@@ -1,4 +1,4 @@
-﻿// Copyright Epic Games, Inc. All Rights Reserved.
+// Copyright Epic Games, Inc. All Rights Reserved.
 
 #pragma once
 
@@ -34,6 +34,7 @@ namespace Chaos
 										 Chaos::FReal                                       InTime,
 										 FPlaybackTickRecord&								TickRecord,
 										 TArray<TPBDRigidParticleHandle<Chaos::FReal, 3>*>& OutUpdatedRigids) const override;
+		virtual void                   WaitForSolverTasks(UPrimitiveComponent* InComponent) const override;
 		// ~End FComponentCacheAdapter interface
 
 	private :
diff --git a/Engine/Source/Runtime/Engine/Classes/Components/SkeletalMeshComponent.h b/Engine/Source/Runtime/Engine/Classes/Components/SkeletalMeshComponent.h
--- a/Engine/Source/Runtime/Engine/Classes/Components/SkeletalMeshComponent.h
+++ b/Engine/Source/Runtime/Engine/Classes/Components/SkeletalMeshComponent.h
@@ -1443,6 +1443,9 @@ public:
 	/** Get the current simulation data map for the clothing on this component. This will stall until the cloth simulation is complete. */
 	const TMap<int32, FClothSimulData>& GetCurrentClothingData_AnyThread() const;
 
+	/** Stalls on any currently running clothing simulations.*/
+	void WaitForExistingParallelClothSimulation_GameThread();
+
 private:
 
 	/** Let the cloth tick and completion tasks have access to private clothing data */
diff --git a/Engine/Source/Runtime/Engine/Private/SkeletalMeshComponentPhysics.cpp b/Engine/Source/Runtime/Engine/Private/SkeletalMeshComponentPhysics.cpp
--- a/Engine/Source/Runtime/Engine/Private/SkeletalMeshComponentPhysics.cpp
+++ b/Engine/Source/Runtime/Engine/Private/SkeletalMeshComponentPhysics.cpp
@@ -3883,6 +3883,13 @@ void USkeletalMeshComponent::GetUpdateClothSimulationData_AnyThread(TMap<int32,
 	}
 }
 
+void USkeletalMeshComponent::WaitForExistingParallelClothSimulation_GameThread()
+{
+	// Should only kick new parallel cloth simulations from game thread, so should be safe to also wait for existing ones there.
+	check(IsInGameThread());
+	HandleExistingParallelClothSimulation();
+}
+
 void USkeletalMeshComponent::DebugDrawClothing(FPrimitiveDrawInterface* PDI)
 {
 #if WITH_EDITOR && ENABLE_DRAW_DEBUG
diff --git a/Engine/Source/Runtime/Experimental/Chaos/Public/Chaos/Framework/PhysicsSolverBase.h b/Engine/Source/Runtime/Experimental/Chaos/Public/Chaos/Framework/PhysicsSolverBase.h
--- a/Engine/Source/Runtime/Experimental/Chaos/Public/Chaos/Framework/PhysicsSolverBase.h
+++ b/Engine/Source/Runtime/Experimental/Chaos/Public/Chaos/Framework/PhysicsSolverBase.h
@@ -174,9 +174,14 @@ namespace Chaos
 	public:
 
 		/** Only allow construction with valid parameters as well as restricting to module construction */
-		virtual ~FPhysicsSolverEvents() { ClearCallbacks(); }
+		virtual ~FPhysicsSolverEvents()
+		{ 
+			EventTeardown.Broadcast();
+			ClearCallbacks(); 
+		}
 		
 		/** Events */
+		/** WARNING: Events are not threadsafe!*/
 		/** Pre advance is called before any physics processing or simulation happens in a given physics update */
 		FDelegateHandle AddPreAdvanceCallback(FSolverPreAdvance::FDelegate InDelegate);
 		bool            RemovePreAdvanceCallback(FDelegateHandle InHandle);
