Chaos : Make Cache collection to be blueprint type

diff --git a/Engine/Plugins/Experimental/ChaosCaching/Source/ChaosCaching/Public/Chaos/CacheCollection.h b/Engine/Plugins/Experimental/ChaosCaching/Source/ChaosCaching/Public/Chaos/CacheCollection.h
--- a/Engine/Plugins/Experimental/ChaosCaching/Source/ChaosCaching/Public/Chaos/CacheCollection.h
+++ b/Engine/Plugins/Experimental/ChaosCaching/Source/ChaosCaching/Public/Chaos/CacheCollection.h
@@ -6,7 +6,7 @@
 
 class UChaosCache;
 
-UCLASS(Experimental)
+UCLASS(Experimental, BlueprintType)
 class CHAOSCACHING_API UChaosCacheCollection : public UObject
 {
 	GENERATED_BODY()
