ChaosCache: Add CompressChannelsData function to allow users to compress FRichCurves. Allow users to not use FPerParticleCacheData to save space.

diff --git a/Engine/Plugins/Experimental/ChaosCaching/Source/ChaosCaching/Private/Chaos/ChaosCache.cpp b/Engine/Plugins/Experimental/ChaosCaching/Source/ChaosCaching/Private/Chaos/ChaosCache.cpp
--- a/Engine/Plugins/Experimental/ChaosCaching/Source/ChaosCaching/Private/Chaos/ChaosCache.cpp
+++ b/Engine/Plugins/Experimental/ChaosCaching/Source/ChaosCaching/Private/Chaos/ChaosCache.cpp
@@ -168,6 +168,23 @@ void UChaosCache::FlushPendingFrames()
 					TargetCurve->RichCurves[ParticleIndex++].AddKey(NewData.Time, ChannelValue);
 				}
 			}
+
+			const bool bContainsParticleData = NewData.PendingParticleData.Num() > 0;
+			// If there is no particle data, use ChannelCurveToParticle instead. 
+			if (!bContainsParticleData)
+			{
+				// If there are more particles, append the curve index to the particle index map
+				const int32 OldNum = ChannelCurveToParticle.Num();
+				const int32 NewNum = NewData.PendingChannelsIndices.Num();	
+				if (OldNum < NewNum)
+				{
+					ChannelCurveToParticle.SetNum(NewNum);
+					for(int32 Index = OldNum; Index < NewNum; ++Index)
+					{
+						ChannelCurveToParticle[Index] = NewData.PendingChannelsIndices[Index];
+					}
+				}
+			}
 		}
 		
 		const int32 ParticleCount = NewData.PendingParticleData.Num();
@@ -235,6 +252,21 @@ void UChaosCache::FlushPendingFrames()
 			Max = FMath::Max(Max, ParticleData.TransformData.GetEndTime());
 		}
 
+		for(const TPair<FName, FRichCurves>& Pair : ChannelsTracks)
+		{
+			const FRichCurves& RichCurves = Pair.Value;
+			for(const FRichCurve& Curve : RichCurves.RichCurves)
+			{
+				float CurveMin, CurveMax;
+				Curve.GetTimeRange(CurveMin, CurveMax);
+				Min = FMath::Min(Min, CurveMin);
+				Max = FMath::Max(Max, CurveMax);
+
+				// assuming all curves have the same time range
+				break;
+			}
+		}
+
 		RecordedDuration = Max - Min;
 	}
 }
@@ -255,6 +287,7 @@ FCacheUserToken UChaosCache::BeginRecord(UPrimitiveComponent* InComponent, FGuid
 			ParticleTracks.Reset();
 			ChannelsTracks.Reset();
 			TrackToParticle.Reset();
+			ChannelCurveToParticle.Reset();
 			CurveData.Reset();
 			EventTracks.Reset();
 			NamedTransformTracks.Reset();
@@ -315,6 +348,11 @@ void UChaosCache::EndRecord(FCacheUserToken& InOutToken)
 			UE_LOG(LogChaosCache, Warning, TEXT("Attempted to close a recording session with an invalid token"));
 		}
 	}
+
+	if (bCompressChannels)
+	{
+		CompressChannelsData(ChannelsCompressionErrorThreshold, ChannelsCompressionSampleRate);
+	}
 }
 
 FCacheUserToken UChaosCache::BeginPlayback()
@@ -414,6 +452,21 @@ FCacheEvaluationResult UChaosCache::Evaluate(const FCacheEvaluationContext& InCo
 					ResultData[CacheIndex] = ChannelData.Value.RichCurves[CacheIndex].Eval(InContext.TickRecord.GetTime(),0.0);
 				}
 			}
+			for(const TPair<FName, FCompressedRichCurves>& CompressedChannelData : CompressedChannelsTracks)
+			{
+				TArray<float>& ResultData = Result.Channels.FindOrAdd(CompressedChannelData.Key);
+				ResultData.Init(0.0, NumProvidedIndices);
+				
+				for(int32 EvalIndex = 0; EvalIndex < NumProvidedIndices; ++EvalIndex)
+				{
+					ResultData[EvalIndex] = CompressedChannelData.Value.CompressedRichCurves[EvalIndex].Eval(InContext.TickRecord.GetTime(),0.0);
+				}
+			}
+			if (ChannelCurveToParticle.Num() > 0)
+			{
+				ensure(ParticleTracks.Num() == 0);
+				Result.ParticleIndices = ChannelCurveToParticle;
+			}
 		}
 
 		for(int32 EvalIndex = 0; EvalIndex < NumProvidedIndices; ++EvalIndex)
@@ -461,17 +514,35 @@ FCacheEvaluationResult UChaosCache::Evaluate(const FCacheEvaluationContext& InCo
 		}
 		if(InContext.bEvaluateChannels)
 		{
+			const int32 NumCurves = NumParticles > 0 ? NumParticles : ChannelCurveToParticle.Num();
 			for(auto& ChannelData : ChannelsTracks)
 			{
 				TArray<float>& ResultData = Result.Channels.FindOrAdd(ChannelData.Key);
-				ResultData.Init(0.0, NumParticles);
+				ResultData.Init(0.0, NumCurves);
 				
-				for(int32 EvalIndex = 0; EvalIndex < NumParticles; ++EvalIndex)
+				for(int32 EvalIndex = 0; EvalIndex < NumCurves; ++EvalIndex)
 				{
 					ResultData[EvalIndex] = ChannelData.Value.RichCurves[EvalIndex].Eval(InContext.TickRecord.GetTime(),0.0);
 				}
 			}
+			for(const TPair<FName, FCompressedRichCurves>& CompressedChannelData : CompressedChannelsTracks)
+			{
+				TArray<float>& ResultData = Result.Channels.FindOrAdd(CompressedChannelData.Key);
+				ResultData.Init(0.0, NumCurves);
+				
+				for(int32 EvalIndex = 0; EvalIndex < NumCurves; ++EvalIndex)
+				{
+					ResultData[EvalIndex] = CompressedChannelData.Value.CompressedRichCurves[EvalIndex].Eval(InContext.TickRecord.GetTime(),0.0);
+				}
+			}
+
+			if (ChannelCurveToParticle.Num() > 0)
+			{
+				ensure(ParticleTracks.Num() == 0);
+				Result.ParticleIndices = ChannelCurveToParticle;
+			}
 		}
+
 		for(int32 Index = 0; Index < NumParticles; ++Index)
 		{
 			if(ParticleTracks[Index].TransformData.BeginOffset > InContext.TickRecord.GetTime())
@@ -622,6 +693,21 @@ void UChaosCache::EvaluateEvents(FPlaybackTickRecord& InTickRecord, TMap<FName,
 	}
 }
 
+void UChaosCache::CompressChannelsData(float ErrorThreshold, float SampleRate)
+{
+	for(TPair<FName, FRichCurves>& ChannelData : ChannelsTracks)
+	{
+		FCompressedRichCurves &CompressedRichCurves = CompressedChannelsTracks.FindOrAdd(ChannelData.Key);
+		const int32 NumCurves = ChannelData.Value.RichCurves.Num();
+		CompressedRichCurves.CompressedRichCurves.SetNum(NumCurves);
+		::ParallelFor(NumCurves, [&ChannelData, &CompressedRichCurves, ErrorThreshold, SampleRate](int32 CurveIndex)
+		{
+			ChannelData.Value.RichCurves[CurveIndex].CompressCurve(CompressedRichCurves.CompressedRichCurves[CurveIndex], ErrorThreshold, SampleRate);
+		}, false);
+	}
+	ChannelsTracks.Reset();
+}
+
 FTransform FParticleTransformTrack::Evaluate(float InCacheTime, const FTransform* MassToLocal) const
 {
 	QUICK_SCOPE_CYCLE_COUNTER(QSTAT_EvalParticleTransformTrack);
diff --git a/Engine/Plugins/Experimental/ChaosCaching/Source/ChaosCaching/Public/Chaos/ChaosCache.h b/Engine/Plugins/Experimental/ChaosCaching/Source/ChaosCaching/Public/Chaos/ChaosCache.h
--- a/Engine/Plugins/Experimental/ChaosCaching/Source/ChaosCaching/Public/Chaos/ChaosCache.h
+++ b/Engine/Plugins/Experimental/ChaosCaching/Source/ChaosCaching/Public/Chaos/ChaosCache.h
@@ -203,7 +203,16 @@ struct FRichCurves
 	TArray<FRichCurve> RichCurves;
 };
 
-struct FPendingFrameWrite
+USTRUCT()
+struct FCompressedRichCurves
+{
+	GENERATED_BODY()
+
+	UPROPERTY()
+	TArray<FCompressedRichCurve> CompressedRichCurves;
+};
+
+struct CHAOSCACHING_API FPendingFrameWrite
 {
 	float                         Time;
 	TArray<FPendingParticleWrite> PendingParticleData;
@@ -371,6 +380,8 @@ public:
 	void EvaluateCurves(const FPerParticleCacheData& InData, float InTime, TMap<FName, float>& OutCurves);
 	void EvaluateEvents(FPlaybackTickRecord& InTickRecord, TMap<FName, TArray<FCacheEventHandle>>& OutEvents);
 
+	void CompressChannelsData(float ErrorThreshold, float SampleRate);
+
 	UPROPERTY(VisibleAnywhere, Category = "Caching")
 	float RecordedDuration;
 
@@ -385,10 +396,17 @@ public:
 	UPROPERTY()
 	TArray<FPerParticleCacheData> ParticleTracks;
 
-	/** Per-particle curve data, continuous per-frame data */
+	/** Map a curve index in the cache to the original particle index specified when recording */
+	UPROPERTY()
+	TArray<int32> ChannelCurveToParticle;
+
+	/** Per-particle data,  continuous per-frame data */
 	UPROPERTY()
 	TMap<FName,FRichCurves> ChannelsTracks;
 
+	UPROPERTY()
+	TMap<FName, FCompressedRichCurves> CompressedChannelsTracks;
+
 	/** Per component/cache curve data, any continuous data that isn't per-particle can be stored here */
 	UPROPERTY()
 	TMap<FName, FRichCurve> CurveData;
@@ -398,6 +416,15 @@ public:
 	UPROPERTY()
 	TMap<FName, FParticleTransformTrack> NamedTransformTracks;
 
+	UPROPERTY()
+	bool bCompressChannels = false;
+
+	UPROPERTY()
+	float ChannelsCompressionErrorThreshold = 1e-5;
+
+	UPROPERTY()
+	float ChannelsCompressionSampleRate = 1.f / 30.f;
+
 	template<typename T>
 	FCacheEventTrack& FindOrAddEventTrack(FName InName)
 	{
diff --git a/Engine/Plugins/Experimental/ChaosClothAsset/Source/ChaosClothAssetEngine/Private/ChaosClothAsset/ClothComponentCacheAdapter.cpp b/Engine/Plugins/Experimental/ChaosClothAsset/Source/ChaosClothAssetEngine/Private/ChaosClothAsset/ClothComponentCacheAdapter.cpp
--- a/Engine/Plugins/Experimental/ChaosClothAsset/Source/ChaosClothAssetEngine/Private/ChaosClothAsset/ClothComponentCacheAdapter.cpp
+++ b/Engine/Plugins/Experimental/ChaosClothAsset/Source/ChaosClothAssetEngine/Private/ChaosClothAsset/ClothComponentCacheAdapter.cpp
@@ -59,7 +59,7 @@ namespace UE::Chaos::ClothAsset
 	bool FClothComponentCacheAdapter::ValidForPlayback(UPrimitiveComponent* InComponent, UChaosCache* InCache) const
 	{
 		const UChaosClothComponent* ClothComp = GetClothComponent(InComponent);
-		return ClothComp && InCache->TrackToParticle.Num() > 0;
+		return ClothComp && InCache->ChannelCurveToParticle.Num() > 0;
 	}
 
 	::Chaos::FPhysicsSolverEvents* FClothComponentCacheAdapter::BuildEventsSolver(UPrimitiveComponent* InComponent) const
