Chaos : Fix cache
- not properly recording if break event is not set on the geometry collection
- wrongly skipping enabled observed components

diff --git a/Engine/Plugins/Experimental/ChaosCaching/Source/ChaosCaching/Public/Chaos/CacheManagerActor.h b/Engine/Plugins/Experimental/ChaosCaching/Source/ChaosCaching/Public/Chaos/CacheManagerActor.h
--- a/Engine/Plugins/Experimental/ChaosCaching/Source/ChaosCaching/Public/Chaos/CacheManagerActor.h
+++ b/Engine/Plugins/Experimental/ChaosCaching/Source/ChaosCaching/Public/Chaos/CacheManagerActor.h
@@ -47,6 +47,7 @@ struct CHAOSCACHING_API FObservedComponent
 	FObservedComponent()
 		: CacheName(NAME_None)
 		, bIsSimulating(true)
+		, bHasNotifyBreaks(false)
 		, bPlaybackEnabled(true)
 		, bHasNotifyBreaks(false)
 		, Cache(nullptr)
