CHaos cache manager - Fix bad merge after auto merge from CL 24571778

diff --git a/Engine/Plugins/Experimental/ChaosCaching/Source/ChaosCaching/Public/Chaos/CacheManagerActor.h b/Engine/Plugins/Experimental/ChaosCaching/Source/ChaosCaching/Public/Chaos/CacheManagerActor.h
--- a/Engine/Plugins/Experimental/ChaosCaching/Source/ChaosCaching/Public/Chaos/CacheManagerActor.h
+++ b/Engine/Plugins/Experimental/ChaosCaching/Source/ChaosCaching/Public/Chaos/CacheManagerActor.h
@@ -47,7 +47,6 @@ struct CHAOSCACHING_API FObservedComponent
 	FObservedComponent()
 		: CacheName(NAME_None)
 		, bIsSimulating(true)
-		, bHasNotifyBreaks(false)
 		, bPlaybackEnabled(true)
 		, bHasNotifyBreaks(false)
 		, Cache(nullptr)
