Chaos Caching: Correctly handle spawnable CacheManagers which are observing components that belong to the cache manager itself (i.e., geometry collections that are created by adding a CacheCollection to a scene).

The crux of the bug is that the SoftComponentRef.OtherActor will point to the spawnable template, not the spawned instance. This is the case for any spawnable observed component.

This fixes the particular case where the Component's owner is the CacheManager itself by setting OtherActor to nullptr, and passing in the CacheManager as the "OwningActor" at runtime. This includes an upgrade in Serialize to clear the OtherActor when it points to the CacheManager itself.

Tested this works with new spawnable cache managers as well as old existing sequences that were previously broken.

diff --git a/Engine/Plugins/Experimental/ChaosCaching/Source/ChaosCaching/Private/Chaos/CacheManagerActor.cpp b/Engine/Plugins/Experimental/ChaosCaching/Source/ChaosCaching/Private/Chaos/CacheManagerActor.cpp
--- a/Engine/Plugins/Experimental/ChaosCaching/Source/ChaosCaching/Private/Chaos/CacheManagerActor.cpp
+++ b/Engine/Plugins/Experimental/ChaosCaching/Source/ChaosCaching/Private/Chaos/CacheManagerActor.cpp
@@ -59,14 +59,14 @@ void FObservedComponent::PostSerialize(const FArchive& Ar)
 	}
 	PRAGMA_ENABLE_DEPRECATION_WARNINGS
 }
-UPrimitiveComponent* FObservedComponent::GetComponent()
+UPrimitiveComponent* FObservedComponent::GetComponent(AActor* OwningActor)
 {
-	return Cast<UPrimitiveComponent>(SoftComponentRef.GetComponent(nullptr));
+	return Cast<UPrimitiveComponent>(SoftComponentRef.GetComponent(OwningActor));
 }
 
-UPrimitiveComponent* FObservedComponent::GetComponent() const
+UPrimitiveComponent* FObservedComponent::GetComponent(AActor* OwningActor) const
 {
-	return Cast<UPrimitiveComponent>(SoftComponentRef.GetComponent(nullptr));
+	return Cast<UPrimitiveComponent>(SoftComponentRef.GetComponent(OwningActor));
 }
 
 AChaosCacheManager::AChaosCacheManager(const FObjectInitializer& ObjectInitializer /*= FObjectInitializer::Get()*/)
@@ -167,9 +167,16 @@ void AChaosCacheManager::Serialize(FArchive& Ar)
 {
 	Super::Serialize(Ar);
 
-	for(FObservedComponent& ObervedComponent : ObservedComponents)
+	for(FObservedComponent& ObservedComponent : ObservedComponents)
 	{
-		ObervedComponent.PostSerialize(Ar);
+		ObservedComponent.PostSerialize(Ar);
+		if (Ar.IsLoading() && ObservedComponent.SoftComponentRef.OtherActor == this)
+		{
+			// Clear OtherActor in order to use OwningActor(this) in when using GetComponent.
+			// This is relevant when this CacheManager is spawnable by sequencer so the
+			// SoftComponentedRef will point to the actual spawned CacheManager instance and not the spawnable template.
+			ObservedComponent.SoftComponentRef.OtherActor = nullptr;
+		}
 	}
 }
 
@@ -188,7 +195,7 @@ void AChaosCacheManager::ResetAllComponentTransforms()
 
 	for(FObservedComponent& Observed : ObservedComponents)
 	{
-		if(UPrimitiveComponent* Comp = Observed.GetComponent())
+		if(UPrimitiveComponent* Comp = Observed.GetComponent(this))
 		{
 			if(UChaosCache* Cache = CacheCollection->FindCache(Observed.CacheName))
 			{
@@ -207,7 +214,7 @@ void AChaosCacheManager::ResetSingleTransform(int32 InIndex)
 
 	FObservedComponent& Observed = ObservedComponents[InIndex];
 
-	if(UPrimitiveComponent* Comp = Observed.GetComponent())
+	if(UPrimitiveComponent* Comp = Observed.GetComponent(this))
 	{
 		if(UChaosCache* Cache = CacheCollection->FindCache(Observed.CacheName))
 		{
@@ -226,7 +233,7 @@ void AChaosCacheManager::SelectComponent(int32 InIndex)
 
 	FObservedComponent& Observed = ObservedComponents[InIndex];
 
-	if(UPrimitiveComponent* Comp = Observed.GetComponent())
+	if(UPrimitiveComponent* Comp = Observed.GetComponent(this))
 	{
 		GEditor->SelectNone(true, true);
 		GEditor->SelectActor(Comp->GetOwner(), true, true);
@@ -284,7 +291,7 @@ void AChaosCacheManager::BeginEvaluate()
 		auto ByPriority = [](const FComponentCacheAdapter* A, const FComponentCacheAdapter* B) {
 			return A->GetPriority() < B->GetPriority();
 		};
-		UPrimitiveComponent* Comp = Observed.GetComponent();
+		UPrimitiveComponent* Comp = Observed.GetComponent(this);
 		
 		if(!Comp)
 		{
@@ -512,7 +519,7 @@ void AChaosCacheManager::HandlePreSolve(Chaos::FReal InDt, Chaos::FPhysicsSolver
 	}
 
 	TickObservedComponents(Data->PlaybackIndices, InDt, [this, Data](UChaosCache* InCache, FObservedComponent& Observed, Chaos::FComponentCacheAdapter* InAdapter) {
-		UPrimitiveComponent* Comp = Observed.GetComponent();
+		UPrimitiveComponent* Comp = Observed.GetComponent(this);
 		if(ensure(Comp))
 		{
 			InAdapter->Playback_PreSolve(Comp, InCache, Observed.TimeSinceTrigger, Observed.TickRecord, Data->PendingKinematicUpdates);
@@ -561,8 +568,8 @@ void AChaosCacheManager::HandlePostSolve(Chaos::FReal InDt, Chaos::FPhysicsSolve
 		return;
 	}
 
-	TickObservedComponents(Data->RecordIndices, InDt, [](UChaosCache* InCache, FObservedComponent& Observed, Chaos::FComponentCacheAdapter* InAdapter) {
-		UPrimitiveComponent* Comp = Observed.GetComponent();
+	TickObservedComponents(Data->RecordIndices, InDt, [this](UChaosCache* InCache, FObservedComponent& Observed, Chaos::FComponentCacheAdapter* InAdapter) {
+		UPrimitiveComponent* Comp = Observed.GetComponent(this);
 		if (ensure(Comp && InCache))
 		{
 			
@@ -599,7 +606,7 @@ void AChaosCacheManager::OnStartFrameChanged(Chaos::FReal InTime)
 	for (int32 ObservedIdx = 0; ObservedIdx < ObservedComponents.Num(); ++ObservedIdx)
 	{
 		FObservedComponent& Observed = ObservedComponents[ObservedIdx];
-		if (UPrimitiveComponent* Comp = Observed.GetComponent())
+		if (UPrimitiveComponent* Comp = Observed.GetComponent(this))
 		{
 			if (!Observed.BestFitAdapter)
 			{
@@ -649,7 +656,7 @@ void AChaosCacheManager::TriggerComponent(UPrimitiveComponent* InComponent)
 {
 	// #BGTODO Maybe not totally thread-safe, probably safer with an atomic or condition var rather than the bTriggered flag
 	FObservedComponent* Found = Algo::FindByPredicate(ObservedComponents, [this, InComponent](const FObservedComponent& Test) {
-		return Test.GetComponent() == InComponent && Test.IsEnabled(CacheMode);
+		return Test.GetComponent(this) == InComponent && Test.IsEnabled(CacheMode);
 	});
 
 	if (Found && StartMode == EStartMode::Triggered)
@@ -661,7 +668,7 @@ void AChaosCacheManager::TriggerComponent(UPrimitiveComponent* InComponent)
 void AChaosCacheManager::TriggerComponentByCache(FName InCacheName)
 {
 	FObservedComponent* Found = Algo::FindByPredicate(ObservedComponents, [this, InCacheName](const FObservedComponent& Test) {
-		return Test.CacheName == InCacheName && Test.GetComponent() && Test.IsEnabled(CacheMode);
+		return Test.CacheName == InCacheName && Test.GetComponent(this) && Test.IsEnabled(CacheMode);
 	});
 
 	if (Found && StartMode == EStartMode::Triggered)
@@ -674,7 +681,7 @@ void AChaosCacheManager::TriggerAll()
 {
 	for(FObservedComponent& Observed : ObservedComponents)
 	{
-		if (StartMode == EStartMode::Triggered && Observed.GetComponent() && Observed.IsEnabled(CacheMode))
+		if (StartMode == EStartMode::Triggered && Observed.GetComponent(this) && Observed.IsEnabled(CacheMode))
 		{
 			Observed.bTriggered = true;
 		}
@@ -704,8 +711,8 @@ void AChaosCacheManager::EnablePlayback(int32 Index, bool bEnable)
 
 FObservedComponent* AChaosCacheManager::FindObservedComponent(UPrimitiveComponent* InComponent)
 {
-	return Algo::FindByPredicate(ObservedComponents, [ToTest = InComponent](const FObservedComponent& Item) {
-		return Item.GetComponent() == ToTest;
+	return Algo::FindByPredicate(ObservedComponents, [this, ToTest = InComponent](const FObservedComponent& Item) {
+		return Item.GetComponent(this) == ToTest;
 	});
 }
 
@@ -715,8 +722,9 @@ FObservedComponent& AChaosCacheManager::AddNewObservedComponent(UPrimitiveCompon
 	ObservedComponents.AddDefaulted();
 	FObservedComponent& NewEntry = ObservedComponents.Last();
 
-	NewEntry.SoftComponentRef.PathToComponent = InComponent->GetPathName(InComponent->GetOwner());
-	NewEntry.SoftComponentRef.OtherActor      = InComponent->GetOwner();
+	AActor* const ComponentOwner = InComponent->GetOwner();
+	NewEntry.SoftComponentRef.PathToComponent = InComponent->GetPathName(ComponentOwner);
+	NewEntry.SoftComponentRef.OtherActor      = ComponentOwner == this ? nullptr : ComponentOwner;
 
 	FName CacheName = NAME_None;
 #if WITH_EDITOR
@@ -746,7 +754,7 @@ void AChaosCacheManager::ClearObservedComponents()
 {
 	for (FObservedComponent& ObservedComponent : ObservedComponents)
 	{
-		UPrimitiveComponent* Comp = ObservedComponent.GetComponent();
+		UPrimitiveComponent* Comp = ObservedComponent.GetComponent(this);
 		Comp->DestroyComponent();
 	}
 	ObservedComponents.Empty();
@@ -794,7 +802,7 @@ void AChaosCacheManager::WaitForObservedComponentSolverTasks()
 		if(const Chaos::FComponentCacheAdapter* Adapter = ActiveAdapters[CompIndex])
 		{
 			FObservedComponent& Observed = ObservedComponents[CompIndex];
-			Adapter->WaitForSolverTasks(Observed.GetComponent());
+			Adapter->WaitForSolverTasks(Observed.GetComponent(this));
 		}
 	}
 }
@@ -804,7 +812,7 @@ void AChaosCacheManager::SetObservedComponentProperties(const ECacheMode& NewCac
 {
 	for (FObservedComponent& ObservedComponent : ObservedComponents)
 	{
-		if (UPrimitiveComponent* PrimComp = ObservedComponent.GetComponent())
+		if (UPrimitiveComponent* PrimComp = ObservedComponent.GetComponent(this))
 		{
 			if (NewCacheMode == ECacheMode::Record)
 			{
diff --git a/Engine/Plugins/Experimental/ChaosCaching/Source/ChaosCaching/Public/Chaos/CacheManagerActor.h b/Engine/Plugins/Experimental/ChaosCaching/Source/ChaosCaching/Public/Chaos/CacheManagerActor.h
--- a/Engine/Plugins/Experimental/ChaosCaching/Source/ChaosCaching/Public/Chaos/CacheManagerActor.h
+++ b/Engine/Plugins/Experimental/ChaosCaching/Source/ChaosCaching/Public/Chaos/CacheManagerActor.h
@@ -125,8 +125,12 @@ struct CHAOSCACHING_API FObservedComponent
 	bool IsEnabled(ECacheMode CacheMode) const;
 
 	/** Gets the component from the internal component ref */
-	UPrimitiveComponent* GetComponent();
-	UPrimitiveComponent* GetComponent() const;
+	UPrimitiveComponent* GetComponent(AActor* OwningActor);
+	UPrimitiveComponent* GetComponent(AActor* OwningActor) const;
+	UE_DEPRECATED(5.3, "Use GetComponent(OwningActor) instead.")
+	UPrimitiveComponent* GetComponent() { return GetComponent(nullptr); }
+	UE_DEPRECATED(5.3, "Use GetComponent(OwningActor) instead.")
+	UPrimitiveComponent* GetComponent() const { return GetComponent(nullptr); }
 
 private:
 	friend class AChaosCacheManager;
diff --git a/Engine/Plugins/Experimental/ChaosCaching/Source/ChaosCachingEditor/Private/Chaos/ChaosCachingEditorPlugin.cpp b/Engine/Plugins/Experimental/ChaosCaching/Source/ChaosCachingEditor/Private/Chaos/ChaosCachingEditorPlugin.cpp
--- a/Engine/Plugins/Experimental/ChaosCaching/Source/ChaosCachingEditor/Private/Chaos/ChaosCachingEditorPlugin.cpp
+++ b/Engine/Plugins/Experimental/ChaosCaching/Source/ChaosCachingEditor/Private/Chaos/ChaosCachingEditorPlugin.cpp
@@ -340,9 +340,7 @@ void IChaosCachingEditorPlugin::OnCreateCacheManager()
 
 				check(Manager);
 
-				const FObservedComponent* Existing = Manager->GetObservedComponents().FindByPredicate([PrimitiveComp](const FObservedComponent& InItem) {
-					return InItem.GetComponent() == PrimitiveComp;
-				});
+				const FObservedComponent* Existing = Manager->FindObservedComponent(PrimitiveComp);
 
 				if(!Existing)
 				{
