Expose Observed component component ref in Chaos cache manager

diff --git a/Engine/Plugins/Experimental/ChaosCaching/Source/ChaosCaching/Public/Chaos/CacheManagerActor.h b/Engine/Plugins/Experimental/ChaosCaching/Source/ChaosCaching/Public/Chaos/CacheManagerActor.h
--- a/Engine/Plugins/Experimental/ChaosCaching/Source/ChaosCaching/Public/Chaos/CacheManagerActor.h
+++ b/Engine/Plugins/Experimental/ChaosCaching/Source/ChaosCaching/Public/Chaos/CacheManagerActor.h
@@ -97,7 +97,7 @@ struct FObservedComponent
 	FComponentReference ComponentRef;
 
 	/** The component observed by this object for either playback or recording */
-	UPROPERTY(EditAnywhere, Category = "Caching", meta = (UseComponentPicker, AllowAnyActor))
+	UPROPERTY(EditAnywhere, BlueprintReadWrite, Category = "Caching", meta = (UseComponentPicker, AllowAnyActor))
 	FSoftComponentReference SoftComponentRef;
 
 	/** Capture of the initial state of the component before cache manager takes control. */
