Fix chaos cache not properly setting the start time on Begin and end play
This fixes an issue where geometry collection will be reset to use their rest collection transform on begin play insteads of respecting matching start time in play or static pose mode

diff --git a/Engine/Plugins/Experimental/ChaosCaching/Source/ChaosCaching/Private/Chaos/CacheManagerActor.cpp b/Engine/Plugins/Experimental/ChaosCaching/Source/ChaosCaching/Private/Chaos/CacheManagerActor.cpp
--- a/Engine/Plugins/Experimental/ChaosCaching/Source/ChaosCaching/Private/Chaos/CacheManagerActor.cpp
+++ b/Engine/Plugins/Experimental/ChaosCaching/Source/ChaosCaching/Private/Chaos/CacheManagerActor.cpp
@@ -75,6 +75,7 @@ AChaosCacheManager::AChaosCacheManager(const FObjectInitializer& ObjectInitializ
 	, StartMode(EStartMode::Timed)
 	, bCanRecord(true)
 	, bIsSimulating(false)
+	, StartTimeAtBeginPlay(0)
 {
 	// This actor will tick, just not normally. There needs to be a tick-like event both before physics simulation
 	// and after physics simulation, we bind to some physics scene events in BeginPlay to handle this.
@@ -319,6 +320,7 @@ void AChaosCacheManager::BeginPlay()
 
 	Super::BeginPlay();
 
+	StartTimeAtBeginPlay = StartTime;
 	BeginEvaluate();
 }
 
@@ -521,6 +523,13 @@ void AChaosCacheManager::BeginEvaluate()
 		FSlateNotificationManager::Get().AddNotification(Info);
 #endif
 	}
+	else
+	{
+		if (CacheMode != ECacheMode::Record)
+		{
+			OnStartFrameChanged(StartTime);
+		}
+	}
 }
 
 void AChaosCacheManager::EndPlay(const EEndPlayReason::Type EndPlayReason)
@@ -528,6 +537,11 @@ void AChaosCacheManager::EndPlay(const EEndPlayReason::Type EndPlayReason)
 	Super::EndPlay(EndPlayReason);
 
 	EndEvaluate();
+
+	if (StartTimeAtBeginPlay != StartTime)
+	{
+		SetStartTime(StartTimeAtBeginPlay);
+	}
 }
 
 void AChaosCacheManager::EndEvaluate()
diff --git a/Engine/Plugins/Experimental/ChaosCaching/Source/ChaosCaching/Public/Chaos/CacheManagerActor.h b/Engine/Plugins/Experimental/ChaosCaching/Source/ChaosCaching/Public/Chaos/CacheManagerActor.h
--- a/Engine/Plugins/Experimental/ChaosCaching/Source/ChaosCaching/Public/Chaos/CacheManagerActor.h
+++ b/Engine/Plugins/Experimental/ChaosCaching/Source/ChaosCaching/Public/Chaos/CacheManagerActor.h
@@ -369,6 +369,9 @@ private:
 	/** Lists of currently open caches that need to be closed when complete */
 	TArray<TTuple<FCacheUserToken, UChaosCache*>> OpenRecordCaches;
 	TArray<TTuple<FCacheUserToken, UChaosCache*>> OpenPlaybackCaches;
+
+	// since the object persists
+	float StartTimeAtBeginPlay;
 };
 
 UCLASS(Experimental, MinimalAPI)
