Chaos cache : fix crash where the cache spawnable template was keeping a reference to the original component parent if it was attached
- fix it by making sure the template is properly detched from the duplicated parent reference

diff --git a/Engine/Plugins/Experimental/ChaosCaching/Source/ChaosCaching/Private/Chaos/ChaosCache.cpp b/Engine/Plugins/Experimental/ChaosCaching/Source/ChaosCaching/Private/Chaos/ChaosCache.cpp
--- a/Engine/Plugins/Experimental/ChaosCaching/Source/ChaosCaching/Private/Chaos/ChaosCache.cpp
+++ b/Engine/Plugins/Experimental/ChaosCaching/Source/ChaosCaching/Private/Chaos/ChaosCache.cpp
@@ -717,6 +717,15 @@ FCacheEvaluationResult UChaosCache::Evaluate(const FCacheEvaluationContext& InCo
 void UChaosCache::BuildSpawnableFromComponent(const UPrimitiveComponent* InComponent, const FTransform& SpaceTransform)
 {
 	Spawnable.DuplicatedTemplate = StaticDuplicateObject(InComponent, this);
+	// duplication of teh component also copy the attach parent reference, but if we keep it this causes crashes because the garbage collection will have a dangling reference
+	// preventing the PIE level from being released properly and causing all sort of problem and crashes after that 
+	if (USceneComponent* SceneComponent = Cast<USceneComponent>(Spawnable.DuplicatedTemplate))
+	{
+		if (SceneComponent->GetAttachParent())
+		{
+			SceneComponent->DetachFromComponent(FDetachmentTransformRules::KeepWorldTransform);
+		}
+	}
 	Spawnable.InitialTransform = InComponent->GetComponentToWorld();
 	Spawnable.ComponentTransform = InComponent->GetComponentToWorld() * SpaceTransform.Inverse();
 }
