Chaos cache : clear the initial fields when creating the spawnable templateto avoid kiiping garbage collected refs

diff --git a/Engine/Plugins/Experimental/ChaosCaching/Source/ChaosCaching/Private/Chaos/ChaosCache.cpp b/Engine/Plugins/Experimental/ChaosCaching/Source/ChaosCaching/Private/Chaos/ChaosCache.cpp
--- a/Engine/Plugins/Experimental/ChaosCaching/Source/ChaosCaching/Private/Chaos/ChaosCache.cpp
+++ b/Engine/Plugins/Experimental/ChaosCaching/Source/ChaosCaching/Private/Chaos/ChaosCache.cpp
@@ -726,6 +726,11 @@ void UChaosCache::BuildSpawnableFromComponent(const UPrimitiveComponent* InCompo
 			SceneComponent->DetachFromComponent(FDetachmentTransformRules::KeepWorldTransform);
 		}
 	}
+	// make sure we do not have any initialization fields because those may also keep dangling references preventing garbage collection
+	if (UGeometryCollectionComponent* GeometryCollectionComponent = Cast<UGeometryCollectionComponent>(Spawnable.DuplicatedTemplate))
+	{
+		GeometryCollectionComponent->InitializationFields.Empty();
+	}
 	Spawnable.InitialTransform = InComponent->GetComponentToWorld();
 	Spawnable.ComponentTransform = InComponent->GetComponentToWorld() * SpaceTransform.Inverse();
 }
