Add simple track compression to chaos cache
- Compression find the keys that are consecutive and of similar transform and collpase them to only 2 when possible
- this is done at the end of recording if the cvar is enabled (p.Chaos.Cache.CompressTracksAfterRecording - on by default)

diff --git a/Engine/Plugins/Experimental/ChaosCaching/Source/ChaosCaching/Private/Chaos/CacheCollection.cpp b/Engine/Plugins/Experimental/ChaosCaching/Source/ChaosCaching/Private/Chaos/CacheCollection.cpp
--- a/Engine/Plugins/Experimental/ChaosCaching/Source/ChaosCaching/Private/Chaos/CacheCollection.cpp
+++ b/Engine/Plugins/Experimental/ChaosCaching/Source/ChaosCaching/Private/Chaos/CacheCollection.cpp
@@ -14,15 +14,31 @@ void UChaosCacheCollection::GetAssetRegistryTags(TArray<FAssetRegistryTag>& OutT
 
 	OutTags.Add(FAssetRegistryTag(TEXT("Number of Observed Components"), FString::FromInt(NumCaches), FAssetRegistryTag::TT_Numerical));
 
-	float Duration = 0.0;
+	float MaxDuration = 0.0;
+	uint32 MaxRecordedFrames = 0;
+	int64 TotalTracks = 0;
+	int64 TotalTransformKeys = 0;
 	for (int32 CacheIdx = 0; CacheIdx < NumCaches; ++CacheIdx)
 	{
 		if (UChaosCache* Cache = Caches[CacheIdx])
 		{
-			Duration = FMath::Max(Duration, Cache->RecordedDuration);
+			MaxDuration = FMath::Max(MaxDuration, Cache->RecordedDuration);
+			MaxRecordedFrames = FMath::Max(MaxRecordedFrames, Cache->NumRecordedFrames);
+			TotalTracks += Cache->ParticleTracks.Num();
+			for (const FPerParticleCacheData& Track : Cache->ParticleTracks)
+			{
+				TotalTransformKeys += Track.TransformData.GetNumKeys();
+			}
 		}
 	}	
-	OutTags.Add(FAssetRegistryTag(TEXT("Recorded Duration"), FString::Printf(TEXT("%.2f"), Duration), FAssetRegistryTag::TT_Numerical));
+	const int64 TransformStorageSizeInBytes = sizeof(FVector3f) + sizeof(FQuat4f) + sizeof(FVector3f) + sizeof(float); // pos + rot + scale + timestamp
+	const double EstimatedTransformMemoryInMegaBytes = double(TotalTransformKeys * TransformStorageSizeInBytes) / 1024.0 / 1024.0;
+
+	OutTags.Add(FAssetRegistryTag(TEXT("Recorded Duration"), FString::Printf(TEXT("%.2f"), MaxDuration), FAssetRegistryTag::TT_Numerical));
+	OutTags.Add(FAssetRegistryTag(TEXT("Recorded Frames"), FString::Printf(TEXT("%d"), MaxRecordedFrames), FAssetRegistryTag::TT_Numerical));
+	OutTags.Add(FAssetRegistryTag(TEXT("Total Tracks"), FString::Printf(TEXT("%lld"), TotalTracks), FAssetRegistryTag::TT_Numerical));
+	OutTags.Add(FAssetRegistryTag(TEXT("Total Transform Keys"), FString::Printf(TEXT("%lld"), TotalTransformKeys), FAssetRegistryTag::TT_Numerical));
+	OutTags.Add(FAssetRegistryTag(TEXT("Estimated Transform Memory"), FString::Printf(TEXT("%.2f MB"), EstimatedTransformMemoryInMegaBytes), FAssetRegistryTag::TT_Numerical));
 #endif
 
 	Super::GetAssetRegistryTags(OutTags);
diff --git a/Engine/Plugins/Experimental/ChaosCaching/Source/ChaosCaching/Private/Chaos/ChaosCache.cpp b/Engine/Plugins/Experimental/ChaosCaching/Source/ChaosCaching/Private/Chaos/ChaosCache.cpp
--- a/Engine/Plugins/Experimental/ChaosCaching/Source/ChaosCaching/Private/Chaos/ChaosCache.cpp
+++ b/Engine/Plugins/Experimental/ChaosCaching/Source/ChaosCaching/Private/Chaos/ChaosCache.cpp
@@ -16,6 +16,13 @@ FAutoConsoleVariableRef CVarChaosCacheUseInterpolation(
 	bChaosCacheUseInterpolation,
 	TEXT("When enabled, cache interpolates between keys.[def: true]"));
 
+bool bChaosCacheCompressTracksAfterRecording = true;
+FAutoConsoleVariableRef CVarChaosCacheCompressTracksAfterRecording(
+	TEXT("p.Chaos.Cache.CompressTracksAfterRecording"),
+	bChaosCacheCompressTracksAfterRecording,
+	TEXT("When enabled, cache will compress the transform tracks after recording is done.[def: true]"));
+
+
 UChaosCache::UChaosCache()
 	: CurrentRecordCount(0)
 	, CurrentPlaybackCount(0)
@@ -461,6 +468,19 @@ void UChaosCache::EndRecord(FCacheUserToken& InOutToken)
 	{
 		CompressChannelsData(ChannelsCompressionErrorThreshold, ChannelsCompressionSampleRate);
 	}
+
+	if (bChaosCacheCompressTracksAfterRecording)
+	{
+		CompressTracks();
+	}
+}
+
+void UChaosCache::CompressTracks()
+{
+	for (FPerParticleCacheData& Track : ParticleTracks)
+	{
+		Track.TransformData.Compress();
+	}
 }
 
 FCacheUserToken UChaosCache::BeginPlayback()
@@ -919,3 +939,66 @@ const float FParticleTransformTrack::GetEndTime() const
 
 	return 0.0f;
 }
+
+void FParticleTransformTrack::Compress()
+{
+	if (KeyTimestamps.Num() > 0)
+	{
+		// simple compression algorithm to remove similar keys
+		// we compare the resulting transform
+		// the compression can be done in place because the number of resulting keys is always smaller than the original number of keys
+
+		int32 CompressedKeyIndex = 1; // set to 1 because we'll always write after KeyIndex
+
+		for (int32 KeyIndex = 0; KeyIndex < KeyTimestamps.Num(); KeyIndex++)
+		{
+			FTransform CurrentTransform = EvaluateAt(KeyIndex);
+
+			// find the next index where the transform is different
+			int32 NextIndex = (KeyIndex + 1);
+			for (; NextIndex < KeyTimestamps.Num(); NextIndex++)
+			{
+				const FTransform NextTransform = EvaluateAt(NextIndex);
+				if (!NextTransform.Equals(CurrentTransform))
+				{
+					// skip write the same value over itself 
+					const int32 LastSimilarIndex = (NextIndex - 1);
+					if (LastSimilarIndex > KeyIndex)
+					{
+						KeyTimestamps[CompressedKeyIndex] = KeyTimestamps[LastSimilarIndex];
+						RawTransformTrack.PosKeys[CompressedKeyIndex] = RawTransformTrack.PosKeys[LastSimilarIndex];
+						RawTransformTrack.RotKeys[CompressedKeyIndex] = RawTransformTrack.RotKeys[LastSimilarIndex];
+						// not sure we use that part anymore ( maybe in older caches ?)
+						RawTransformTrack.ScaleKeys[CompressedKeyIndex] = RawTransformTrack.ScaleKeys[LastSimilarIndex];
+					}
+					CompressedKeyIndex++;
+					break;
+				}
+			}
+			// we we have reached the end we need to copy the last key
+			if (NextIndex == KeyTimestamps.Num())
+			{
+				const int32 LastSimilarIndex = (NextIndex - 1);
+				KeyTimestamps[CompressedKeyIndex] = KeyTimestamps[LastSimilarIndex];
+				RawTransformTrack.PosKeys[CompressedKeyIndex] = RawTransformTrack.PosKeys[LastSimilarIndex];
+				RawTransformTrack.RotKeys[CompressedKeyIndex] = RawTransformTrack.RotKeys[LastSimilarIndex];
+				// not sure we use that part anymore ( maybe in older caches ?)
+				RawTransformTrack.ScaleKeys[CompressedKeyIndex] = RawTransformTrack.ScaleKeys[LastSimilarIndex];
+				
+				// we are now done 
+				CompressedKeyIndex++;
+				break; 
+			}
+			// make sure we start back as far as we can
+			KeyIndex = (NextIndex - 1);
+		}
+
+		// we are done we can now shrink the original arrays to the compressed size
+		const int32 CompressedSize = CompressedKeyIndex;
+		KeyTimestamps.SetNum(CompressedSize);
+		RawTransformTrack.PosKeys.SetNum(CompressedSize);
+		RawTransformTrack.RotKeys.SetNum(CompressedSize);
+		// not sure we use that part anymore ( maybe in older caches ?)
+		RawTransformTrack.ScaleKeys.SetNum(CompressedSize);
+	}
+}
diff --git a/Engine/Plugins/Experimental/ChaosCaching/Source/ChaosCaching/Public/Chaos/ChaosCache.h b/Engine/Plugins/Experimental/ChaosCaching/Source/ChaosCaching/Public/Chaos/ChaosCache.h
--- a/Engine/Plugins/Experimental/ChaosCaching/Source/ChaosCaching/Public/Chaos/ChaosCache.h
+++ b/Engine/Plugins/Experimental/ChaosCaching/Source/ChaosCaching/Public/Chaos/ChaosCache.h
@@ -65,6 +65,8 @@ struct FParticleTransformTrack
 	const float GetDuration() const;
 	const float GetBeginTime() const;
 	const float GetEndTime() const;
+
+	void Compress();
 };
 
 USTRUCT()
@@ -444,6 +446,8 @@ private:
 	template<EQueueMode Mode>
 	bool FlushPendingFrames_MainPass(TQueue<FPendingFrameWrite, Mode>& InPendingWrites, bool bCanSimpleCopyChannelData);
 
+	void CompressTracks();
+
 	friend class AChaosCacheManager;
 
 	/** Timestamped generic event tracks */
