Chaos cache : make sure we only compress tracks that have at least 3 keys

diff --git a/Engine/Plugins/Experimental/ChaosCaching/Source/ChaosCaching/Private/Chaos/ChaosCache.cpp b/Engine/Plugins/Experimental/ChaosCaching/Source/ChaosCaching/Private/Chaos/ChaosCache.cpp
--- a/Engine/Plugins/Experimental/ChaosCaching/Source/ChaosCaching/Private/Chaos/ChaosCache.cpp
+++ b/Engine/Plugins/Experimental/ChaosCaching/Source/ChaosCaching/Private/Chaos/ChaosCache.cpp
@@ -942,7 +942,8 @@ const float FParticleTransformTrack::GetEndTime() const
 
 void FParticleTransformTrack::Compress()
 {
-	if (KeyTimestamps.Num() > 0)
+	// we only need to compress if there's more than 3 keys
+	if (KeyTimestamps.Num() >= 3)
 	{
 		// simple compression algorithm to remove similar keys
 		// we compare the resulting transform
