Chaos cache : disable compression for the time being as it can result in crashes in some cases

diff --git a/Engine/Plugins/Experimental/ChaosCaching/Source/ChaosCaching/Private/Chaos/ChaosCache.cpp b/Engine/Plugins/Experimental/ChaosCaching/Source/ChaosCaching/Private/Chaos/ChaosCache.cpp
--- a/Engine/Plugins/Experimental/ChaosCaching/Source/ChaosCaching/Private/Chaos/ChaosCache.cpp
+++ b/Engine/Plugins/Experimental/ChaosCaching/Source/ChaosCaching/Private/Chaos/ChaosCache.cpp
@@ -16,7 +16,7 @@ FAutoConsoleVariableRef CVarChaosCacheUseInterpolation(
 	bChaosCacheUseInterpolation,
 	TEXT("When enabled, cache interpolates between keys.[def: true]"));
 
-bool bChaosCacheCompressTracksAfterRecording = true;
+bool bChaosCacheCompressTracksAfterRecording = false;
 FAutoConsoleVariableRef CVarChaosCacheCompressTracksAfterRecording(
 	TEXT("p.Chaos.Cache.CompressTracksAfterRecording"),
 	bChaosCacheCompressTracksAfterRecording,
