Add a SetCurrentTime blueprint function to Chaos Cache manager

diff --git a/Engine/Plugins/Experimental/ChaosCaching/Source/ChaosCaching/Private/Chaos/CacheManagerActor.cpp b/Engine/Plugins/Experimental/ChaosCaching/Source/ChaosCaching/Private/Chaos/CacheManagerActor.cpp
--- a/Engine/Plugins/Experimental/ChaosCaching/Source/ChaosCaching/Private/Chaos/CacheManagerActor.cpp
+++ b/Engine/Plugins/Experimental/ChaosCaching/Source/ChaosCaching/Private/Chaos/CacheManagerActor.cpp
@@ -257,6 +257,13 @@ void AChaosCacheManager::SetStartTime(float InStartTime)
 	OnStartFrameChanged(InStartTime);
 }
 
+void AChaosCacheManager::SetCurrentTime(float CurrentTime)
+{
+	// follow to SetStartTime as this is the how we animate the cache 
+	// todo(chaos) : we should probably separate the concept of start time and current time 
+	SetStartTime(CurrentTime);
+}
+
 void AChaosCacheManager::ResetAllComponentTransforms()
 {
 	if(!CacheCollection)
diff --git a/Engine/Plugins/Experimental/ChaosCaching/Source/ChaosCaching/Public/Chaos/CacheManagerActor.h b/Engine/Plugins/Experimental/ChaosCaching/Source/ChaosCaching/Public/Chaos/CacheManagerActor.h
--- a/Engine/Plugins/Experimental/ChaosCaching/Source/ChaosCaching/Public/Chaos/CacheManagerActor.h
+++ b/Engine/Plugins/Experimental/ChaosCaching/Source/ChaosCaching/Public/Chaos/CacheManagerActor.h
@@ -225,6 +225,9 @@ public:
 	UFUNCTION(CallInEditor)
 	CHAOSCACHING_API void SetStartTime(float InStartTime);
 
+	UFUNCTION(BlueprintCallable, Category = "Caching")
+	void SetCurrentTime(float CurrentTime);
+
 	/** 
 	 * Resets all components back to the world space transform they had when the cache for them was originally recorded
 	 * if one is available
