GC caching : fix crash in when compressing tracks

diff --git a/Engine/Plugins/Experimental/ChaosCaching/Source/ChaosCaching/Private/Chaos/ChaosCache.cpp b/Engine/Plugins/Experimental/ChaosCaching/Source/ChaosCaching/Private/Chaos/ChaosCache.cpp
--- a/Engine/Plugins/Experimental/ChaosCaching/Source/ChaosCaching/Private/Chaos/ChaosCache.cpp
+++ b/Engine/Plugins/Experimental/ChaosCaching/Source/ChaosCaching/Private/Chaos/ChaosCache.cpp
@@ -976,8 +976,8 @@ void FParticleTransformTrack::Compress()
 					break;
 				}
 			}
-			// we we have reached the end we need to copy the last key
-			if (NextIndex == KeyTimestamps.Num())
+			// we we have reached the end and haven't copied the last key we need to do it here 
+			if (NextIndex == KeyTimestamps.Num() && CompressedKeyIndex < KeyTimestamps.Num())
 			{
 				const int32 LastSimilarIndex = (NextIndex - 1);
 				KeyTimestamps[CompressedKeyIndex] = KeyTimestamps[LastSimilarIndex];
@@ -994,12 +994,15 @@ void FParticleTransformTrack::Compress()
 			KeyIndex = (NextIndex - 1);
 		}
 
-		// we are done we can now shrink the original arrays to the compressed size
+		// we are done we can now shrink the original arrays to the compressed size if necessary
 		const int32 CompressedSize = CompressedKeyIndex;
-		KeyTimestamps.SetNum(CompressedSize);
-		RawTransformTrack.PosKeys.SetNum(CompressedSize);
-		RawTransformTrack.RotKeys.SetNum(CompressedSize);
-		// not sure we use that part anymore ( maybe in older caches ?)
-		RawTransformTrack.ScaleKeys.SetNum(CompressedSize);
+		if (CompressedSize < KeyTimestamps.Num())
+		{
+			KeyTimestamps.SetNum(CompressedSize);
+			RawTransformTrack.PosKeys.SetNum(CompressedSize);
+			RawTransformTrack.RotKeys.SetNum(CompressedSize);
+			// not sure we use that part anymore ( maybe in older caches ?)
+			RawTransformTrack.ScaleKeys.SetNum(CompressedSize);
+		}
 	}
 }
