GC Caching : fix slow down when scrubbing start time property while having the sequecner opened

diff --git a/Engine/Plugins/Experimental/ChaosCaching/Source/ChaosCaching/Public/Chaos/CacheManagerActor.h b/Engine/Plugins/Experimental/ChaosCaching/Source/ChaosCaching/Public/Chaos/CacheManagerActor.h
--- a/Engine/Plugins/Experimental/ChaosCaching/Source/ChaosCaching/Public/Chaos/CacheManagerActor.h
+++ b/Engine/Plugins/Experimental/ChaosCaching/Source/ChaosCaching/Public/Chaos/CacheManagerActor.h
@@ -198,7 +198,7 @@ public:
 	* Defines the (random access) time that represents the rest pose of the components managed by this cache.
 	* When in Play mode, the components are set to the state provided by the caches at this evaluated time.
 	*/
-	UPROPERTY(EditAnywhere, Interp, BlueprintReadWrite, Category = "Caching", meta=(SequencerTrackClass="MovieSceneFloatTrack"))
+	UPROPERTY(EditAnywhere, Interp, BlueprintReadWrite, Category = "Caching", meta=(SequencerTrackClass="/Script/MovieSceneTracks.MovieSceneFloatTrack"))
 	float StartTime;
 
 	/** AActor interface */
