Enable Chaos Cache Track compression by default

diff --git a/Engine/Plugins/Experimental/ChaosCaching/Source/ChaosCaching/Private/Chaos/ChaosCache.cpp b/Engine/Plugins/Experimental/ChaosCaching/Source/ChaosCaching/Private/Chaos/ChaosCache.cpp
--- a/Engine/Plugins/Experimental/ChaosCaching/Source/ChaosCaching/Private/Chaos/ChaosCache.cpp
+++ b/Engine/Plugins/Experimental/ChaosCaching/Source/ChaosCaching/Private/Chaos/ChaosCache.cpp
@@ -16,7 +16,7 @@ FAutoConsoleVariableRef CVarChaosCacheUseInterpolation(
 	bChaosCacheUseInterpolation,
 	TEXT("When enabled, cache interpolates between keys.[def: true]"));
 
-bool bChaosCacheCompressTracksAfterRecording = false;
+bool bChaosCacheCompressTracksAfterRecording = true;
 FAutoConsoleVariableRef CVarChaosCacheCompressTracksAfterRecording(
 	TEXT("p.Chaos.Cache.CompressTracksAfterRecording"),
 	bChaosCacheCompressTracksAfterRecording,
