Change filetype from utf8 to text to be consistent with rest of engine

diff --git a/Engine/Plugins/Experimental/ChaosCaching/Source/ChaosCaching/Public/Chaos/Sequencer/MovieSceneChaosCacheSection.h b/Engine/Plugins/Experimental/ChaosCaching/Source/ChaosCaching/Public/Chaos/Sequencer/MovieSceneChaosCacheSection.h
--- a/Engine/Plugins/Experimental/ChaosCaching/Source/ChaosCaching/Public/Chaos/Sequencer/MovieSceneChaosCacheSection.h
+++ b/Engine/Plugins/Experimental/ChaosCaching/Source/ChaosCaching/Public/Chaos/Sequencer/MovieSceneChaosCacheSection.h
@@ -1,4 +1,4 @@
-﻿// Copyright Epic Games, Inc. All Rights Reserved.
+// Copyright Epic Games, Inc. All Rights Reserved.
 
 #pragma once
 
diff --git a/Engine/Plugins/FX/Niagara/Source/Niagara/Classes/NiagaraVersionedObject.h b/Engine/Plugins/FX/Niagara/Source/Niagara/Classes/NiagaraVersionedObject.h
--- a/Engine/Plugins/FX/Niagara/Source/Niagara/Classes/NiagaraVersionedObject.h
+++ b/Engine/Plugins/FX/Niagara/Source/Niagara/Classes/NiagaraVersionedObject.h
@@ -1,4 +1,4 @@
-﻿// Copyright Epic Games, Inc. All Rights Reserved.
+// Copyright Epic Games, Inc. All Rights Reserved.
 
 #pragma once
 
diff --git a/Engine/Source/Runtime/Engine/Classes/Debug/DebugDrawComponent.h b/Engine/Source/Runtime/Engine/Classes/Debug/DebugDrawComponent.h
--- a/Engine/Source/Runtime/Engine/Classes/Debug/DebugDrawComponent.h
+++ b/Engine/Source/Runtime/Engine/Classes/Debug/DebugDrawComponent.h
@@ -1,4 +1,4 @@
-﻿// Copyright Epic Games, Inc. All Rights Reserved.
+// Copyright Epic Games, Inc. All Rights Reserved.
 #pragma once
 
 #include "Components/PrimitiveComponent.h"
diff --git a/Engine/Source/Runtime/Engine/Classes/Kismet/KismetStringTableLibrary.h b/Engine/Source/Runtime/Engine/Classes/Kismet/KismetStringTableLibrary.h
--- a/Engine/Source/Runtime/Engine/Classes/Kismet/KismetStringTableLibrary.h
+++ b/Engine/Source/Runtime/Engine/Classes/Kismet/KismetStringTableLibrary.h
@@ -1,4 +1,4 @@
-﻿// Copyright Epic Games, Inc. All Rights Reserved.
+// Copyright Epic Games, Inc. All Rights Reserved.
 
 #pragma once
 
diff --git a/Engine/Source/Runtime/Experimental/Animation/Constraints/Public/ConstraintsActor.h b/Engine/Source/Runtime/Experimental/Animation/Constraints/Public/ConstraintsActor.h
--- a/Engine/Source/Runtime/Experimental/Animation/Constraints/Public/ConstraintsActor.h
+++ b/Engine/Source/Runtime/Experimental/Animation/Constraints/Public/ConstraintsActor.h
@@ -1,4 +1,4 @@
-﻿// Copyright Epic Games, Inc. All Rights Reserved.
+// Copyright Epic Games, Inc. All Rights Reserved.
 
 #pragma once
 
diff --git a/Engine/Source/Runtime/InteractiveToolsFramework/Public/TargetInterfaces/SkeletalMeshBackedTarget.h b/Engine/Source/Runtime/InteractiveToolsFramework/Public/TargetInterfaces/SkeletalMeshBackedTarget.h
--- a/Engine/Source/Runtime/InteractiveToolsFramework/Public/TargetInterfaces/SkeletalMeshBackedTarget.h
+++ b/Engine/Source/Runtime/InteractiveToolsFramework/Public/TargetInterfaces/SkeletalMeshBackedTarget.h
@@ -1,4 +1,4 @@
-﻿// Copyright Epic Games, Inc. All Rights Reserved.
+// Copyright Epic Games, Inc. All Rights Reserved.
 
 #pragma once
 
diff --git a/Engine/Source/Runtime/UMG/Public/Components/WindowTitleBarArea.h b/Engine/Source/Runtime/UMG/Public/Components/WindowTitleBarArea.h
--- a/Engine/Source/Runtime/UMG/Public/Components/WindowTitleBarArea.h
+++ b/Engine/Source/Runtime/UMG/Public/Components/WindowTitleBarArea.h
@@ -1,4 +1,4 @@
-﻿// Copyright Epic Games, Inc. All Rights Reserved.
+// Copyright Epic Games, Inc. All Rights Reserved.
 
 #pragma once
 
