Fix Cache manager CIS

diff --git a/Engine/Plugins/Experimental/ChaosCaching/Source/ChaosCaching/Private/Chaos/CacheManagerActor.cpp b/Engine/Plugins/Experimental/ChaosCaching/Source/ChaosCaching/Private/Chaos/CacheManagerActor.cpp
--- a/Engine/Plugins/Experimental/ChaosCaching/Source/ChaosCaching/Private/Chaos/CacheManagerActor.cpp
+++ b/Engine/Plugins/Experimental/ChaosCaching/Source/ChaosCaching/Private/Chaos/CacheManagerActor.cpp
@@ -846,19 +846,23 @@ FObservedComponent& AChaosCacheManager::AddNewObservedComponent(UPrimitiveCompon
 	return NewEntry;
 }
 
-FObservedComponent& AChaosCacheManager::FindOrAddObservedComponent(UPrimitiveComponent* InComponent, const FName& CacheName, const bool bTransferSimulation)
+void AChaosCacheManager::FindOrAddObservedComponent(UPrimitiveComponent* InComponent, const FName& CacheName, const bool bTransferSimulation)
 {
-	FObservedComponent* FoundComponent = FindObservedComponent(InComponent);
-	FObservedComponent* ObservedComponent = FoundComponent ? FoundComponent : &AddNewObservedComponent(InComponent);
+	if(InComponent != nullptr)
+	{
+		FObservedComponent* FoundComponent = FindObservedComponent(InComponent);
+		FObservedComponent* ObservedComponent = FoundComponent ? FoundComponent : &AddNewObservedComponent(InComponent);
 
-	ObservedComponent->bIsSimulating = (InComponent && bTransferSimulation) ? InComponent->BodyInstance.bSimulatePhysics : false;
+		if(ObservedComponent)
+		{
+			ObservedComponent->bIsSimulating = bTransferSimulation ? InComponent->BodyInstance.bSimulatePhysics : false;
 
-	if(CacheName != TEXT(""))
-	{
-		ObservedComponent->CacheName = CacheName;
+			if(CacheName != TEXT(""))
+			{
+				ObservedComponent->CacheName = CacheName;
+			}
+		}
 	}
-
-	return *ObservedComponent;
 }
 
 void AChaosCacheManager::RemoveObservedComponent(UPrimitiveComponent* PrimitiveComponent)
diff --git a/Engine/Plugins/Experimental/ChaosCaching/Source/ChaosCaching/Public/Chaos/CacheManagerActor.h b/Engine/Plugins/Experimental/ChaosCaching/Source/ChaosCaching/Public/Chaos/CacheManagerActor.h
--- a/Engine/Plugins/Experimental/ChaosCaching/Source/ChaosCaching/Public/Chaos/CacheManagerActor.h
+++ b/Engine/Plugins/Experimental/ChaosCaching/Source/ChaosCaching/Public/Chaos/CacheManagerActor.h
@@ -269,7 +269,7 @@ public:
 	TArray<FObservedComponent>& GetObservedComponents() {return ObservedComponents;}
 	
 	/** Find or add a primitive component to a cache manager */
-	CHAOSCACHING_API FObservedComponent& FindOrAddObservedComponent(UPrimitiveComponent* InComponent, const FName& CacheName = TEXT(""), const bool bTransferSimulationFlag = false);
+	CHAOSCACHING_API void FindOrAddObservedComponent(UPrimitiveComponent* InComponent, const FName& CacheName = TEXT(""), const bool bTransferSimulationFlag = false);
 
 	/** Remove a primitive component from the cache manager */
 	CHAOSCACHING_API void RemoveObservedComponent(UPrimitiveComponent* InComponent);
