Fix Chaos Cache compression
Current code could skip some keys in the track in some specific cases
- Simplified the logic
- Extracted out the key copying code for better readibility

diff --git a/Engine/Plugins/Experimental/ChaosCaching/Source/ChaosCaching/Private/Chaos/ChaosCache.cpp b/Engine/Plugins/Experimental/ChaosCaching/Source/ChaosCaching/Private/Chaos/ChaosCache.cpp
--- a/Engine/Plugins/Experimental/ChaosCaching/Source/ChaosCaching/Private/Chaos/ChaosCache.cpp
+++ b/Engine/Plugins/Experimental/ChaosCaching/Source/ChaosCaching/Private/Chaos/ChaosCache.cpp
@@ -951,6 +951,30 @@ const float FParticleTransformTrack::GetEndTime() const
 	return 0.0f;
 }
 
+void FParticleTransformTrack::CopyTrackEntry(int32 FromIndex, int32 ToIndex)
+{
+	if (FromIndex != ToIndex)
+	{
+		KeyTimestamps[ToIndex] = KeyTimestamps[FromIndex];
+		RawTransformTrack.PosKeys[ToIndex] = RawTransformTrack.PosKeys[FromIndex];
+		RawTransformTrack.RotKeys[ToIndex] = RawTransformTrack.RotKeys[FromIndex];
+		// not sure we use that part anymore ( maybe in older caches ?)
+		RawTransformTrack.ScaleKeys[ToIndex] = RawTransformTrack.ScaleKeys[FromIndex];
+	}
+}
+
+void FParticleTransformTrack::ResizeTrack(int32 NewSize)
+{
+	if (NewSize < KeyTimestamps.Num())
+	{
+		KeyTimestamps.SetNum(NewSize);
+		RawTransformTrack.PosKeys.SetNum(NewSize);
+		RawTransformTrack.RotKeys.SetNum(NewSize);
+		// not sure we use that part anymore ( maybe in older caches ?)
+		RawTransformTrack.ScaleKeys.SetNum(NewSize);
+	}
+}
+
 void FParticleTransformTrack::Compress()
 {
 	// we only need to compress if there's more than 3 keys
@@ -959,13 +983,19 @@ void FParticleTransformTrack::Compress()
 		// simple compression algorithm to remove similar keys
 		// we compare the resulting transform
 		// the compression can be done in place because the number of resulting keys is always smaller than the original number of keys
+		// when more than 3 consecutive transforms are similar they are compressed using two ( earlier and oldest timestamp )
+		// [A1 A2 A3 A4 B1 C1 C2 C3 D1 E1 E2 E3] will result in [A1 A4 B1 C1 C3 D1 E1 E3 ] 
 
-		int32 CompressedKeyIndex = 1; // set to 1 because we'll always write after KeyIndex
+		int32 CompressedKeyIndex = 0; // index to write to
 
 		for (int32 KeyIndex = 0; KeyIndex < KeyTimestamps.Num(); KeyIndex++)
 		{
 			FTransform CurrentTransform = EvaluateAt(KeyIndex);
 
+			// write current 
+			CopyTrackEntry(KeyIndex, CompressedKeyIndex);
+			CompressedKeyIndex++;
+
 			// find the next index where the transform is different
 			int32 NextIndex = (KeyIndex + 1);
 			for (; NextIndex < KeyTimestamps.Num(); NextIndex++)
@@ -973,47 +1003,26 @@ void FParticleTransformTrack::Compress()
 				const FTransform NextTransform = EvaluateAt(NextIndex);
 				if (!NextTransform.Equals(CurrentTransform))
 				{
-					// skip write the same value over itself 
-					const int32 LastSimilarIndex = (NextIndex - 1);
-					if (LastSimilarIndex > KeyIndex)
-					{
-						KeyTimestamps[CompressedKeyIndex] = KeyTimestamps[LastSimilarIndex];
-						RawTransformTrack.PosKeys[CompressedKeyIndex] = RawTransformTrack.PosKeys[LastSimilarIndex];
-						RawTransformTrack.RotKeys[CompressedKeyIndex] = RawTransformTrack.RotKeys[LastSimilarIndex];
-						// not sure we use that part anymore ( maybe in older caches ?)
-						RawTransformTrack.ScaleKeys[CompressedKeyIndex] = RawTransformTrack.ScaleKeys[LastSimilarIndex];
-					}
-					CompressedKeyIndex++;
 					break;
 				}
 			}
-			// we we have reached the end and haven't copied the last key we need to do it here 
-			if (NextIndex == KeyTimestamps.Num() && CompressedKeyIndex < KeyTimestamps.Num())
+
+			// if there's at least one identical transform before the new one 
+			const int32 LastSimilarIndex = (NextIndex - 1);
+			if (LastSimilarIndex > KeyIndex)
 			{
-				const int32 LastSimilarIndex = (NextIndex - 1);
-				KeyTimestamps[CompressedKeyIndex] = KeyTimestamps[LastSimilarIndex];
-				RawTransformTrack.PosKeys[CompressedKeyIndex] = RawTransformTrack.PosKeys[LastSimilarIndex];
-				RawTransformTrack.RotKeys[CompressedKeyIndex] = RawTransformTrack.RotKeys[LastSimilarIndex];
-				// not sure we use that part anymore ( maybe in older caches ?)
-				RawTransformTrack.ScaleKeys[CompressedKeyIndex] = RawTransformTrack.ScaleKeys[LastSimilarIndex];
-				
-				// we are now done 
-				CompressedKeyIndex++;
-				break; 
+				// write same transform with oldest timestamp
+				if (CompressedKeyIndex < KeyTimestamps.Num())
+				{
+					CopyTrackEntry(LastSimilarIndex, CompressedKeyIndex);
+					CompressedKeyIndex++;
+				}
+				KeyIndex = (NextIndex - 1); // account for the automatic increment of the loop
 			}
-			// make sure we start back as far as we can
-			KeyIndex = (NextIndex - 1);
 		}
 
-		// we are done we can now shrink the original arrays to the compressed size if necessary
+		// we are done we can now shrink the original arrays to the compressed size
 		const int32 CompressedSize = CompressedKeyIndex;
-		if (CompressedSize < KeyTimestamps.Num())
-		{
-			KeyTimestamps.SetNum(CompressedSize);
-			RawTransformTrack.PosKeys.SetNum(CompressedSize);
-			RawTransformTrack.RotKeys.SetNum(CompressedSize);
-			// not sure we use that part anymore ( maybe in older caches ?)
-			RawTransformTrack.ScaleKeys.SetNum(CompressedSize);
-		}
+		ResizeTrack(CompressedSize);
 	}
 }
diff --git a/Engine/Plugins/Experimental/ChaosCaching/Source/ChaosCaching/Public/Chaos/ChaosCache.h b/Engine/Plugins/Experimental/ChaosCaching/Source/ChaosCaching/Public/Chaos/ChaosCache.h
--- a/Engine/Plugins/Experimental/ChaosCaching/Source/ChaosCaching/Public/Chaos/ChaosCache.h
+++ b/Engine/Plugins/Experimental/ChaosCaching/Source/ChaosCaching/Public/Chaos/ChaosCache.h
@@ -67,6 +67,10 @@ struct FParticleTransformTrack
 	const float GetEndTime() const;
 
 	void Compress();
+
+private:
+	void CopyTrackEntry(int32 FromIndex, int32 ToIndex);
+	void ResizeTrack(int32 NewSize);
 };
 
 USTRUCT()
