Chaos Cache: Add the ability to change the interpolation mode when replaying cache
Caches containing fast spinning parts, can replay weird when scrubbing in editor ( and likely when usiong large time deltas )
It is now possible to change those caches to use dual quaternion interpolation instead of normal FTransform blend  to resolve the issue
this should also solve slow mo cases as well

diff --git a/Engine/Plugins/Experimental/ChaosCaching/Source/ChaosCaching/Private/Chaos/CacheCollection.cpp b/Engine/Plugins/Experimental/ChaosCaching/Source/ChaosCaching/Private/Chaos/CacheCollection.cpp
--- a/Engine/Plugins/Experimental/ChaosCaching/Source/ChaosCaching/Private/Chaos/CacheCollection.cpp
+++ b/Engine/Plugins/Experimental/ChaosCaching/Source/ChaosCaching/Private/Chaos/CacheCollection.cpp
@@ -137,3 +137,30 @@ uint32 UChaosCacheCollection::GetMaxNumFrames() const
 	}
 	return MaxFramesCount;
 }
+
+#if WITH_EDITOR
+void UChaosCacheCollection::PostEditChangeProperty(struct FPropertyChangedEvent& PropertyChangedEvent)
+{
+	Super::PostEditChangeProperty(PropertyChangedEvent);
+
+	const FName PropertyName = PropertyChangedEvent.GetPropertyName();
+	if (PropertyName == GET_MEMBER_NAME_STRING_CHECKED(UChaosCacheCollection, InterpolationMode))
+	{
+		SetInterpolationMode(InterpolationMode);
+	}
+}
+#endif
+
+void UChaosCacheCollection::SetInterpolationMode(EChaosCacheInterpolationMode Mode)
+{
+	InterpolationMode = Mode;
+	for (UChaosCache* CacheInstance : Caches)
+	{
+		if (CacheInstance)
+		{
+			CacheInstance->InterpolationMode = Mode;
+		}
+	}
+}
+
+
diff --git a/Engine/Plugins/Experimental/ChaosCaching/Source/ChaosCaching/Private/Chaos/ChaosCache.cpp b/Engine/Plugins/Experimental/ChaosCaching/Source/ChaosCaching/Private/Chaos/ChaosCache.cpp
--- a/Engine/Plugins/Experimental/ChaosCaching/Source/ChaosCaching/Private/Chaos/ChaosCache.cpp
+++ b/Engine/Plugins/Experimental/ChaosCaching/Source/ChaosCaching/Private/Chaos/ChaosCache.cpp
@@ -7,6 +7,7 @@
 #include "UObject/Package.h"
 #include "GeometryCollection/GeometryCollectionComponent.h"
 #include "GeometryCollection/GeometryCollectionObject.h"
+#include "Math/DualQuat.h"
 
 #include UE_INLINE_GENERATED_CPP_BY_NAME(ChaosCache)
 
@@ -734,7 +735,7 @@ FCacheEvaluationResult UChaosCache::Evaluate(const FCacheEvaluationContext& InCo
 		Result.NamedTransforms.Reset();
 		for (const TPair<FName, FNamedTransformTrack>& NamedTransform : NamedTransformTracks)
 		{
-			Result.NamedTransforms.Add(NamedTransform.Key, NamedTransform.Value.Evaluate(InContext.TickRecord.GetTime(), nullptr));
+			Result.NamedTransforms.Add(NamedTransform.Key, NamedTransform.Value.Evaluate(InContext.TickRecord.GetTime(), nullptr, InterpolationMode));
 		}
 	}
 
@@ -793,7 +794,7 @@ void UChaosCache::EvaluateSingle(int32 InIndex, FPlaybackTickRecord& InTickRecor
 
 void UChaosCache::EvaluateTransform(const FPerParticleCacheData& InData, float InTime, const FTransform* MassToLocal, FTransform& OutTransform)
 {
-	OutTransform = InData.TransformData.Evaluate(InTime, MassToLocal);
+	OutTransform = InData.TransformData.Evaluate(InTime, MassToLocal, InterpolationMode);
 }
 
 void UChaosCache::EvaluateCurves(const FPerParticleCacheData& InData, float InTime, TMap<FName, float>& OutCurves)
@@ -861,7 +862,7 @@ void UChaosCache::CompressChannelsData(float ErrorThreshold, float SampleRate)
 	ChannelsTracks.Reset();
 }
 
-FTransform FParticleTransformTrack::Evaluate(float InCacheTime, const FTransform* MassToLocal) const
+FTransform FParticleTransformTrack::Evaluate(float InCacheTime, const FTransform* MassToLocal, EChaosCacheInterpolationMode InterpolationMode) const
 {
 	QUICK_SCOPE_CYCLE_COUNTER(QSTAT_EvalParticleTransformTrack);
 	const int32 NumKeys = GetNumKeys();
@@ -886,7 +887,29 @@ FTransform FParticleTransformTrack::Evaluate(float InCacheTime, const FTransform
 
 		if (bChaosCacheUseInterpolation)
 		{
-			Result.Blend(TransformA, TransformB, Alpha);
+			if (InterpolationMode == EChaosCacheInterpolationMode::QuatInterp)
+			{
+				Result.Blend(TransformA, TransformB, Alpha);
+			}
+			else if (InterpolationMode == EChaosCacheInterpolationMode::EulerInterp)
+			{
+				Result.SetTranslation(FMath::Lerp(TransformA.GetTranslation(), TransformB.GetTranslation(), Alpha));
+				Result.SetScale3D(FMath::Lerp(TransformA.GetScale3D(), TransformB.GetScale3D(), Alpha));
+
+				const FRotator RotatorA = TransformA.Rotator();
+				const FRotator RotatorB = TransformB.Rotator();
+				const FRotator DeltaRotator = RotatorB - RotatorA;
+				Result.SetRotation(FQuat(RotatorA + Alpha * DeltaRotator));
+			}
+			else // (InterpolationMode == EChaosCacheInterpolationMode::DualQuatInterp)
+			{
+				FTransform AdjustedTransformB(TransformB);
+				if ((TransformB.GetRotation() | TransformA.GetRotation()) < 0.0f)
+				{
+					AdjustedTransformB.SetRotation(TransformB.GetRotation() * -1.0f);
+				}
+				Result = (FDualQuat(TransformA) * (1 - Alpha) + FDualQuat(AdjustedTransformB) * Alpha).Normalized().AsFTransform(FMath::Lerp(TransformA.GetScale3D(), AdjustedTransformB.GetScale3D(), Alpha));
+			}
 		}
 		else
 		{
diff --git a/Engine/Plugins/Experimental/ChaosCaching/Source/ChaosCaching/Public/Chaos/CacheCollection.h b/Engine/Plugins/Experimental/ChaosCaching/Source/ChaosCaching/Public/Chaos/CacheCollection.h
--- a/Engine/Plugins/Experimental/ChaosCaching/Source/ChaosCaching/Public/Chaos/CacheCollection.h
+++ b/Engine/Plugins/Experimental/ChaosCaching/Source/ChaosCaching/Public/Chaos/CacheCollection.h
@@ -2,6 +2,8 @@
 
 #pragma once
 
+#include "Chaos/ChaosCacheInterpolationMode.h"
+
 #include "CacheCollection.generated.h"
 
 class UChaosCache;
@@ -16,6 +18,7 @@ public:
 	CHAOSCACHING_API virtual void GetAssetRegistryTags(FAssetRegistryTagsContext Context) const override;
 	UE_DEPRECATED(5.4, "Implement the version that takes FAssetRegistryTagsContext instead.")
 	CHAOSCACHING_API virtual void GetAssetRegistryTags(TArray<FAssetRegistryTag>& OutTags) const override;
+	CHAOSCACHING_API virtual void PostEditChangeProperty(struct FPropertyChangedEvent& PropertyChangedEvent);
 	//~ End UObject Interface
 	
 	CHAOSCACHING_API UChaosCache* FindCache(const FName& CacheName) const;
@@ -31,6 +34,15 @@ public:
 	/** Return the max number of frames of all the caches stored in the collection */
 	CHAOSCACHING_API uint32 GetMaxNumFrames() const;
 
+	/** Set the interpolation mode on this cache collection */
+	CHAOSCACHING_API void SetInterpolationMode(EChaosCacheInterpolationMode Mode);
+
+	CHAOSCACHING_API EChaosCacheInterpolationMode GetInterpolationMode() const { return InterpolationMode; }
+
 	UPROPERTY(EditAnywhere, Instanced, Category="Caching", meta=(EditFixedOrder))
 	TArray<TObjectPtr<UChaosCache>> Caches;
+
+private:
+	UPROPERTY(EditAnywhere, Category = "Caching")
+	EChaosCacheInterpolationMode InterpolationMode;
 };
diff --git a/Engine/Plugins/Experimental/ChaosCaching/Source/ChaosCaching/Public/Chaos/ChaosCache.h b/Engine/Plugins/Experimental/ChaosCaching/Source/ChaosCaching/Public/Chaos/ChaosCache.h
--- a/Engine/Plugins/Experimental/ChaosCaching/Source/ChaosCaching/Public/Chaos/ChaosCache.h
+++ b/Engine/Plugins/Experimental/ChaosCaching/Source/ChaosCaching/Public/Chaos/ChaosCache.h
@@ -7,9 +7,11 @@
 #include "Containers/Queue.h"
 #include "Curves/RichCurve.h"
 
+#include "Chaos/ChaosCacheInterpolationMode.h"
 
 #include "ChaosCache.generated.h"
 
+
 USTRUCT()
 struct FParticleTransformTrack
 {
@@ -43,7 +45,7 @@ struct FParticleTransformTrack
 	 * @param InCacheTime Absolute time from the beginning of the entire owning cache to evaluate.
 	 * @param MassToLocal if not null, the will be premultiplied to transform before interpolation
 	 */
-	FTransform Evaluate(float InCacheTime, const FTransform* MassToLocal) const;
+	FTransform Evaluate(float InCacheTime, const FTransform* MassToLocal, EChaosCacheInterpolationMode InterpolationMode =  EChaosCacheInterpolationMode::QuatInterp) const;
 
 	/**
 	 * Find the index the key where timestamp is directly above InCacheTime
@@ -394,6 +396,9 @@ public:
 	UPROPERTY(VisibleAnywhere, Category = "Caching")
 	uint32 NumRecordedFrames;
 
+	UPROPERTY(VisibleAnywhere, Category = "Caching")
+	EChaosCacheInterpolationMode InterpolationMode = EChaosCacheInterpolationMode::QuatInterp;
+
 	/** Maps a track index in the cache to the original particle index specified when recording */
 	UPROPERTY()
 	TArray<int32> TrackToParticle;
diff --git a/Engine/Plugins/Experimental/ChaosCaching/Source/ChaosCaching/Public/Chaos/ChaosCacheInterpolationMode.h b/Engine/Plugins/Experimental/ChaosCaching/Source/ChaosCaching/Public/Chaos/ChaosCacheInterpolationMode.h
new file mode 100644
--- /dev/null
+++ b/Engine/Plugins/Experimental/ChaosCaching/Source/ChaosCaching/Public/Chaos/ChaosCacheInterpolationMode.h
@@ -0,0 +1,18 @@
+// Copyright Epic Games, Inc. All Rights Reserved.
+
+#pragma once
+
+#include "ChaosCacheInterpolationMode.generated.h"
+
+UENUM()
+enum class EChaosCacheInterpolationMode : uint8
+{
+	/** Shortest Path or Quaternion interpolation for the rotation. */
+	QuatInterp,
+
+	/** Rotor or Euler Angle interpolation. */
+	EulerInterp,
+
+	/** Dual quaternion interpolation, follows helix or screw-motion path between keyframes.   */
+	DualQuatInterp
+};
\ No newline at end of file
