Fix non-editor
#rnx

diff --git a/Engine/Plugins/Experimental/ChaosCaching/Source/ChaosCaching/Public/Chaos/CacheCollection.h b/Engine/Plugins/Experimental/ChaosCaching/Source/ChaosCaching/Public/Chaos/CacheCollection.h
--- a/Engine/Plugins/Experimental/ChaosCaching/Source/ChaosCaching/Public/Chaos/CacheCollection.h
+++ b/Engine/Plugins/Experimental/ChaosCaching/Source/ChaosCaching/Public/Chaos/CacheCollection.h
@@ -18,7 +18,9 @@ public:
 	CHAOSCACHING_API virtual void GetAssetRegistryTags(FAssetRegistryTagsContext Context) const override;
 	UE_DEPRECATED(5.4, "Implement the version that takes FAssetRegistryTagsContext instead.")
 	CHAOSCACHING_API virtual void GetAssetRegistryTags(TArray<FAssetRegistryTag>& OutTags) const override;
+#if WITH_EDITOR
 	CHAOSCACHING_API virtual void PostEditChangeProperty(struct FPropertyChangedEvent& PropertyChangedEvent);
+#endif
 	//~ End UObject Interface
 	
 	CHAOSCACHING_API UChaosCache* FindCache(const FName& CacheName) const;
