Fix crash when changing a cache name in cache collection properties

diff --git a/Engine/Plugins/Experimental/ChaosCaching/Source/ChaosCachingEditor/Private/Chaos/CacheCollectionCustomization.cpp b/Engine/Plugins/Experimental/ChaosCaching/Source/ChaosCachingEditor/Private/Chaos/CacheCollectionCustomization.cpp
--- a/Engine/Plugins/Experimental/ChaosCaching/Source/ChaosCachingEditor/Private/Chaos/CacheCollectionCustomization.cpp
+++ b/Engine/Plugins/Experimental/ChaosCaching/Source/ChaosCachingEditor/Private/Chaos/CacheCollectionCustomization.cpp
@@ -132,14 +132,17 @@ void FCacheCollectionDetails::OnChangeCacheName(const FText& InNewName, int32 In
 		return;
 	}
 
-	TSharedPtr<SEditableTextBox> TextBox = NameEditBoxes[InIndex];
-	if(!IsValidName(Collection, FName(InNewName.ToString())))
-	{
-		TextBox->SetError(LOCTEXT("InvalidNameError", "Invalid Cache Name"));
-	}
-	else
+
+	if (TSharedPtr<SEditableTextBox> TextBox = NameEditBoxes[InIndex])
 	{
-		TextBox->SetError(TEXT(""));
+		if (!IsValidName(Collection, FName(InNewName.ToString())))
+		{
+			TextBox->SetError(LOCTEXT("InvalidNameError", "Invalid Cache Name"));
+		}
+		else
+		{
+			TextBox->SetError(TEXT(""));
+		}
 	}
 }
 
