Fix crash when using Reset transform on a Cache manager that has not Cache collection assigned

diff --git a/Engine/Plugins/Experimental/ChaosCaching/Source/ChaosCaching/Private/Chaos/CacheManagerActor.cpp b/Engine/Plugins/Experimental/ChaosCaching/Source/ChaosCaching/Private/Chaos/CacheManagerActor.cpp
--- a/Engine/Plugins/Experimental/ChaosCaching/Source/ChaosCaching/Private/Chaos/CacheManagerActor.cpp
+++ b/Engine/Plugins/Experimental/ChaosCaching/Source/ChaosCaching/Private/Chaos/CacheManagerActor.cpp
@@ -285,6 +285,11 @@ void AChaosCacheManager::ResetAllComponentTransforms()
 
 void AChaosCacheManager::ResetSingleTransform(int32 InIndex)
 {
+	if (!CacheCollection)
+	{
+		return;
+	}
+
 	if(!ObservedComponents.IsValidIndex(InIndex))
 	{
 		return;
