Fix crash in Cache collection when editing an empty cache entry

diff --git a/Engine/Plugins/Experimental/ChaosCaching/Source/ChaosCachingEditor/Private/Chaos/CacheCollectionCustomization.cpp b/Engine/Plugins/Experimental/ChaosCaching/Source/ChaosCachingEditor/Private/Chaos/CacheCollectionCustomization.cpp
--- a/Engine/Plugins/Experimental/ChaosCaching/Source/ChaosCachingEditor/Private/Chaos/CacheCollectionCustomization.cpp
+++ b/Engine/Plugins/Experimental/ChaosCaching/Source/ChaosCachingEditor/Private/Chaos/CacheCollectionCustomization.cpp
@@ -160,7 +160,10 @@ void FCacheCollectionDetails::OnCommitCacheName(const FText& InNewName, ETextCom
 		if(IsValidName(Collection, NewName))
 		{
 			UChaosCache* CurrCache = Collection->Caches[InIndex];
-			CurrCache->Rename(*NewName.ToString());
+			if (CurrCache)
+			{
+				CurrCache->Rename(*NewName.ToString());
+			}
 		}
 	}
 }
