AChaosCacheManager: add a flag to show/hide the Billboard component. Set it hidden in the DataflowSimulationScene

diff --git a/Engine/Plugins/Experimental/ChaosCaching/Source/ChaosCaching/Private/Chaos/CacheManagerActor.cpp b/Engine/Plugins/Experimental/ChaosCaching/Source/ChaosCaching/Private/Chaos/CacheManagerActor.cpp
--- a/Engine/Plugins/Experimental/ChaosCaching/Source/ChaosCaching/Private/Chaos/CacheManagerActor.cpp
+++ b/Engine/Plugins/Experimental/ChaosCaching/Source/ChaosCaching/Private/Chaos/CacheManagerActor.cpp
@@ -104,16 +104,17 @@ AChaosCacheManager::AChaosCacheManager(const FObjectInitializer& ObjectInitializ
 		};
 		static FConstructorStatics ConstructorStatics;
 
-		UBillboardComponent* SpriteComp = ObjectInitializer.CreateEditorOnlyDefaultSubobject<UBillboardComponent>(this, TEXT("Editor Icon"));
+		EditorIconComponent = ObjectInitializer.CreateEditorOnlyDefaultSubobject<UBillboardComponent>(this, TEXT("Editor Icon"));
 
-		if (SpriteComp)
+		if (EditorIconComponent)
 		{
-			SpriteComp->Sprite = ConstructorStatics.SpriteTextureObject.Get();
-			SpriteComp->SpriteInfo.Category = ConstructorStatics.ID_CacheManager;
-			SpriteComp->SpriteInfo.DisplayName = ConstructorStatics.NAME_CacheManager;
-			SpriteComp->Mobility = EComponentMobility::Static;
-			SpriteComp->AttachToComponent(RootComponent, FAttachmentTransformRules::KeepRelativeTransform);
+			EditorIconComponent->Sprite = ConstructorStatics.SpriteTextureObject.Get();
+			EditorIconComponent->SpriteInfo.Category = ConstructorStatics.ID_CacheManager;
+			EditorIconComponent->SpriteInfo.DisplayName = ConstructorStatics.NAME_CacheManager;
+			EditorIconComponent->Mobility = EComponentMobility::Static;
+			EditorIconComponent->AttachToComponent(RootComponent, FAttachmentTransformRules::KeepRelativeTransform);
 		}
+
 	}
 #endif
 }
@@ -962,7 +963,16 @@ void AChaosCacheManager::SetObservedComponentProperties(const ECacheMode& NewCac
 		}
 	}
 }
-#endif
+
+void AChaosCacheManager::SetEditorIconVisibility(bool bVisibility)
+{
+	if (EditorIconComponent)
+	{
+		EditorIconComponent->SetVisibility(bVisibility);
+	}
+}
+
+#endif	// WITH_EDITOR
 
 AChaosCachePlayer::AChaosCachePlayer(const FObjectInitializer& ObjectInitializer /*= FObjectInitializer::Get()*/)
 	: Super(ObjectInitializer)
diff --git a/Engine/Plugins/Experimental/ChaosCaching/Source/ChaosCaching/Public/Chaos/CacheManagerActor.h b/Engine/Plugins/Experimental/ChaosCaching/Source/ChaosCaching/Public/Chaos/CacheManagerActor.h
--- a/Engine/Plugins/Experimental/ChaosCaching/Source/ChaosCaching/Public/Chaos/CacheManagerActor.h
+++ b/Engine/Plugins/Experimental/ChaosCaching/Source/ChaosCaching/Public/Chaos/CacheManagerActor.h
@@ -13,6 +13,7 @@ namespace Chaos { class FPhysicsSolverEvents; }
 
 class UChaosCacheCollection;
 class UPrimitiveComponent;
+class UBillboardComponent;
 
 namespace Chaos
 {
@@ -279,6 +280,8 @@ public:
 	
 #if WITH_EDITOR
 	CHAOSCACHING_API void SetObservedComponentProperties(const ECacheMode& NewCacheMode);
+
+	CHAOSCACHING_API void SetEditorIconVisibility(bool bVisibility);
 #endif
 	
 
@@ -387,6 +390,12 @@ private:
 
 	// since the object persists
 	float StartTimeAtBeginPlay;
+
+#if WITH_EDITORONLY_DATA
+	UPROPERTY(Transient)
+	TObjectPtr<UBillboardComponent> EditorIconComponent;
+#endif
+
 };
 
 UCLASS(Experimental, MinimalAPI)
diff --git a/Engine/Plugins/Experimental/Dataflow/Source/DataflowEditor/Private/Dataflow/DataflowSimulationScene.cpp b/Engine/Plugins/Experimental/Dataflow/Source/DataflowEditor/Private/Dataflow/DataflowSimulationScene.cpp
--- a/Engine/Plugins/Experimental/Dataflow/Source/DataflowEditor/Private/Dataflow/DataflowSimulationScene.cpp
+++ b/Engine/Plugins/Experimental/Dataflow/Source/DataflowEditor/Private/Dataflow/DataflowSimulationScene.cpp
@@ -38,7 +38,9 @@ FDataflowSimulationScene::FDataflowSimulationScene(FPreviewScene::ConstructionVa
 	SceneDescription->SetSimulationScene(this);
 
 	SimulationGenerator = MakeShared<UE::Dataflow::FDataflowSimulationGenerator>();
-	RootSceneActor = GetWorld()->SpawnActor<AChaosCacheManager>();
+	AChaosCacheManager* const CacheManager = GetWorld()->SpawnActor<AChaosCacheManager>();
+	CacheManager->SetEditorIconVisibility(false);
+	RootSceneActor = CacheManager;
 
 	if(GetEditorContent())
 	{
