Add missing orientation to chaos breaking and trailing events
- Without orientation the the extend of the bounds ( in local space ) have little use especially when the ratio between axis is quite large  ( flat fragment for example )

diff --git a/Engine/Plugins/Experimental/ChaosCaching/Source/ChaosCaching/Private/Chaos/Adapters/GeometryCollectionComponentCacheAdapter.cpp b/Engine/Plugins/Experimental/ChaosCaching/Source/ChaosCaching/Private/Chaos/Adapters/GeometryCollectionComponentCacheAdapter.cpp
--- a/Engine/Plugins/Experimental/ChaosCaching/Source/ChaosCaching/Private/Chaos/Adapters/GeometryCollectionComponentCacheAdapter.cpp
+++ b/Engine/Plugins/Experimental/ChaosCaching/Source/ChaosCaching/Private/Chaos/Adapters/GeometryCollectionComponentCacheAdapter.cpp
@@ -360,6 +360,7 @@ namespace Chaos
 						FBreakingData CachedBreak;
 						CachedBreak.Proxy = Proxy;
 						CachedBreak.Location = ComponentToWorld.TransformPosition(Event->Location);
+						CachedBreak.Orientation = ComponentToWorld.TransformRotation(Event->Orientation);
 						CachedBreak.Velocity = ComponentToWorld.TransformVector(Event->Velocity);
 						CachedBreak.AngularVelocity = Event->AngularVelocity;
 						CachedBreak.Mass = Event->Mass;
@@ -404,6 +405,7 @@ namespace Chaos
 						FTrailingData CachedTrail;
 						CachedTrail.Proxy = Proxy;
 						CachedTrail.Location = ComponentToWorld.TransformPosition(Event->Location);
+						CachedTrail.Orientation = ComponentToWorld.TransformRotation(Event->Orientation);
 						CachedTrail.Velocity = ComponentToWorld.TransformVector(Event->Velocity);
 						CachedTrail.AngularVelocity = Event->AngularVelocity;
 						CachedTrail.BoundingBox = TAABB<FReal, 3>(Event->BoundingBoxMin, Event->BoundingBoxMax);
diff --git a/Engine/Plugins/Experimental/ChaosCaching/Source/ChaosCaching/Public/Chaos/Adapters/GeometryCollectionComponentCacheAdapter.h b/Engine/Plugins/Experimental/ChaosCaching/Source/ChaosCaching/Public/Chaos/Adapters/GeometryCollectionComponentCacheAdapter.h
--- a/Engine/Plugins/Experimental/ChaosCaching/Source/ChaosCaching/Public/Chaos/Adapters/GeometryCollectionComponentCacheAdapter.h
+++ b/Engine/Plugins/Experimental/ChaosCaching/Source/ChaosCaching/Public/Chaos/Adapters/GeometryCollectionComponentCacheAdapter.h
@@ -44,6 +44,7 @@ struct FBreakingEvent : public FCacheEventBase
 	FBreakingEvent()
 		: Index(INDEX_NONE)
 		, Location(0.0f, 0.0f, 0.0f)
+		, Orientation(FQuat::Identity)
 		, Velocity(0.0f, 0.0f, 0.0f)
 		, AngularVelocity(0.0f, 0.0f, 0.0f)
 		, Mass(1.0f)
@@ -54,6 +55,7 @@ struct FBreakingEvent : public FCacheEventBase
 	FBreakingEvent(int32 InIndex, const Chaos::FBreakingData& InData, const FTransform& WorldToComponent)
 		: Index(InIndex)
 		, Location(WorldToComponent.TransformPosition(InData.Location))
+		, Orientation(WorldToComponent.TransformRotation(InData.Orientation))
 		, Velocity(WorldToComponent.TransformVector(InData.Velocity))
 		, AngularVelocity(InData.AngularVelocity)
 		, Mass(InData.Mass)
@@ -71,6 +73,9 @@ struct FBreakingEvent : public FCacheEventBase
 	UPROPERTY()
 	FVector Location;
 	
+	UPROPERTY()
+	FQuat Orientation;
+
 	UPROPERTY()
 	FVector Velocity;
 
@@ -166,11 +171,12 @@ struct FTrailingEvent : public FCacheEventBase
 {
 	GENERATED_BODY()
 
-		static FName EventName;
+	static FName EventName;
 
 	FTrailingEvent()
 		: Index(INDEX_NONE)
 		, Location(0.0f, 0.0f, 0.0f)
+		, Orientation(FQuat::Identity)
 		, Velocity(0.0f, 0.0f, 0.0f)
 		, AngularVelocity(0.0f, 0.0f, 0.0f)
 		, BoundingBoxMin(0.0f, 0.0f, 0.0f)
@@ -180,6 +186,7 @@ struct FTrailingEvent : public FCacheEventBase
 	FTrailingEvent(int32 InIndex, const Chaos::FTrailingData& InData, const FTransform& WorldToComponent)
 		: Index(InIndex)
 		, Location(WorldToComponent.TransformPosition(InData.Location))
+		, Orientation(WorldToComponent.TransformRotation(InData.Orientation))
 		, Velocity(WorldToComponent.TransformVector(InData.Velocity))
 		, AngularVelocity(InData.AngularVelocity)
 		, BoundingBoxMin(InData.BoundingBox.Min())
@@ -191,22 +198,25 @@ struct FTrailingEvent : public FCacheEventBase
 	}
 
 	UPROPERTY()
-		int32 Index;
+	int32 Index;
 
 	UPROPERTY()
-		FVector Location;
+	FVector Location;
 
 	UPROPERTY()
-		FVector Velocity;
+	FQuat Orientation;
 
 	UPROPERTY()
-		FVector AngularVelocity;
+	FVector Velocity;
 
 	UPROPERTY()
-		FVector BoundingBoxMin;
+	FVector AngularVelocity;
 
 	UPROPERTY()
-		FVector BoundingBoxMax;
+	FVector BoundingBoxMin;
+
+	UPROPERTY()
+	FVector BoundingBoxMax;
 };
 
 
diff --git a/Engine/Source/Runtime/Engine/Private/PhysicsEngine/Experimental/ChaosEventType.cpp b/Engine/Source/Runtime/Engine/Private/PhysicsEngine/Experimental/ChaosEventType.cpp
--- a/Engine/Source/Runtime/Engine/Private/PhysicsEngine/Experimental/ChaosEventType.cpp
+++ b/Engine/Source/Runtime/Engine/Private/PhysicsEngine/Experimental/ChaosEventType.cpp
@@ -47,6 +47,7 @@ FCollisionChaosEvent::FCollisionChaosEvent(const Chaos::FCollidingData& Collisio
 FChaosBreakEvent::FChaosBreakEvent()
 	: Component(nullptr)
 	, Location(FVector::ZeroVector)
+	, Orientation(FQuat::Identity)
 	, Velocity(FVector::ZeroVector)
 	, AngularVelocity(FVector::ZeroVector)
 	, Extents(FVector::ZeroVector)
@@ -59,6 +60,7 @@ FChaosBreakEvent::FChaosBreakEvent()
 FChaosBreakEvent::FChaosBreakEvent(const Chaos::FBreakingData& BreakingData)
 	: Component(nullptr)
 	, Location(BreakingData.Location)
+	, Orientation(BreakingData.Orientation)
 	, Velocity(BreakingData.Velocity)
 	, AngularVelocity(BreakingData.AngularVelocity)
 	, Extents(BreakingData.BoundingBox.Extents())
diff --git a/Engine/Source/Runtime/Engine/Public/Physics/Experimental/ChaosEventType.h b/Engine/Source/Runtime/Engine/Public/Physics/Experimental/ChaosEventType.h
--- a/Engine/Source/Runtime/Engine/Public/Physics/Experimental/ChaosEventType.h
+++ b/Engine/Source/Runtime/Engine/Public/Physics/Experimental/ChaosEventType.h
@@ -93,10 +93,14 @@ public:
 	UPROPERTY(BlueprintReadOnly, Category = "Break Event")
 	TObjectPtr<UPrimitiveComponent> Component = nullptr;
 
-	/** World location of the break */
+	/** World location of the breaking piece */
 	UPROPERTY(BlueprintReadOnly, Category = "Break Event")
 	FVector Location;
 
+	/** World orientation of the breaking piece */
+	UPROPERTY(BlueprintReadOnly, Category = "Break Event")
+	FQuat Orientation;
+
 	/** Linear Velocity of the breaking particle  */
 	UPROPERTY(BlueprintReadOnly, Category = "Break Event")
 	FVector Velocity;
@@ -105,7 +109,7 @@ public:
 	UPROPERTY(BlueprintReadOnly, Category = "Break Event")
 	FVector AngularVelocity;
 
-	/** Extents of the bounding box */
+	/** Extents of the bounding box in local space of the breaking piece */
 	UPROPERTY(BlueprintReadOnly, Category = "Break Event")
 	FVector Extents;
 
@@ -160,35 +164,35 @@ public:
 
 	/** primitive component involved in the crumble event */
 	UPROPERTY(BlueprintReadOnly, Category = "Crumble Event")
-		TObjectPtr<UPrimitiveComponent> Component = nullptr;
+	TObjectPtr<UPrimitiveComponent> Component = nullptr;
 
 	/** World location of the crumbling cluster */
 	UPROPERTY(BlueprintReadOnly, Category = "Crumble Event")
-		FVector Location;
+	FVector Location;
 
 	/** World orientation of the crumbling cluster */
 	UPROPERTY(BlueprintReadOnly, Category = "Crumble Event")
-		FQuat Orientation;
+	FQuat Orientation;
 
 	/** Linear Velocity of the crumbling cluster */
 	UPROPERTY(BlueprintReadOnly, Category = "Crumble Event")
-		FVector LinearVelocity;
+	FVector LinearVelocity;
 
 	/** Angular Velocity of the crumbling cluster  */
 	UPROPERTY(BlueprintReadOnly, Category = "Crumble Event")
-		FVector AngularVelocity;
+	FVector AngularVelocity;
 
 	/** Mass of the crumbling cluster  */
 	UPROPERTY(BlueprintReadOnly, Category = "Crumble Event")
-		float Mass;
+	float Mass;
 
 	/** Local bounding box of the crumbling cluster  */
 	UPROPERTY(BlueprintReadOnly, Category = "Crumble Event")
-		FBox LocalBounds;
+	FBox LocalBounds;
 
 	/** List of children indices released (optional : see geometry collection component bCrumblingEventIncludesChildren) */
 	UPROPERTY(BlueprintReadOnly, Category = "Crumble Event")
-		TArray<int32> Children;
+	TArray<int32> Children;
 };
 
 
diff --git a/Engine/Source/Runtime/Experimental/Chaos/Private/Chaos/PBDRigidClustering.cpp b/Engine/Source/Runtime/Experimental/Chaos/Private/Chaos/PBDRigidClustering.cpp
--- a/Engine/Source/Runtime/Experimental/Chaos/Private/Chaos/PBDRigidClustering.cpp
+++ b/Engine/Source/Runtime/Experimental/Chaos/Private/Chaos/PBDRigidClustering.cpp
@@ -766,6 +766,7 @@ namespace Chaos
 				FBreakingData& ClusterBreak = MAllClusterBreakings.AddDefaulted_GetRef();
 				ClusterBreak.Proxy = ClusteredParticle->PhysicsProxy();
 				ClusterBreak.Location = ClusteredParticle->GetX();
+				ClusterBreak.Orientation = ClusteredParticle->GetR();
 				ClusterBreak.Velocity = ClusteredParticle->GetV();
 				ClusterBreak.AngularVelocity = ClusteredParticle->GetW();
 				ClusterBreak.Mass = ClusteredParticle->M();
diff --git a/Engine/Source/Runtime/Experimental/Chaos/Private/EventDefaults.cpp b/Engine/Source/Runtime/Experimental/Chaos/Private/EventDefaults.cpp
--- a/Engine/Source/Runtime/Experimental/Chaos/Private/EventDefaults.cpp
+++ b/Engine/Source/Runtime/Experimental/Chaos/Private/EventDefaults.cpp
@@ -400,6 +400,7 @@ namespace Chaos
 						{
 							FTrailingData TrailingData;
 							TrailingData.Location = ActiveParticle->GetX();
+							TrailingData.Orientation = ActiveParticle->GetR();
 							TrailingData.Velocity = ActiveParticle->GetV();
 							TrailingData.AngularVelocity = ActiveParticle->GetW();
 							TrailingData.Mass = ActiveParticle->M();
diff --git a/Engine/Source/Runtime/Experimental/Chaos/Public/Chaos/ExternalCollisionData.h b/Engine/Source/Runtime/Experimental/Chaos/Public/Chaos/ExternalCollisionData.h
--- a/Engine/Source/Runtime/Experimental/Chaos/Public/Chaos/ExternalCollisionData.h
+++ b/Engine/Source/Runtime/Experimental/Chaos/Public/Chaos/ExternalCollisionData.h
@@ -215,6 +215,7 @@ namespace Chaos
 			: FBaseEventFlag()
 			, Proxy(nullptr)
 			, Location(FVec3((FReal)0.0))
+			, Orientation(FRotation3::Identity)
 			, Velocity(FVec3((FReal)0.0))
 			, AngularVelocity(FVec3((FReal)0.0))
 			, Mass((FReal)0.0)
@@ -231,6 +232,7 @@ namespace Chaos
 		IPhysicsProxyBase* Proxy;
 		
 		FVec3 Location;
+		FRotation3 Orientation;
 		FVec3 Velocity;
 		FVec3 AngularVelocity;
 		FReal Mass;
@@ -346,6 +348,7 @@ namespace Chaos
 	{
 		FTrailingData()
 			: Location(FVec3((FReal)0.0))
+			, Orientation(FRotation3::Identity)
 			, Velocity(FVec3((FReal)0.0))
 			, AngularVelocity(FVec3((FReal)0.0))
 			, Mass((FReal)0.0)
@@ -353,6 +356,7 @@ namespace Chaos
 			, BoundingBox(FAABB3(FVec3((FReal)0.0), FVec3((FReal)0.0)))
 		{}
 
+		UE_DEPRECATED(5.6, "Deprecated constructor")
 		FTrailingData(FVec3 InLocation, FVec3 InVelocity, FVec3 InAngularVelocity, FReal InMass
 			, IPhysicsProxyBase* InProxy, Chaos::TAABB<FReal, 3>& InBoundingBox)
 			: Location(InLocation)
@@ -365,6 +369,7 @@ namespace Chaos
 		{}
 
 		FVec3 Location;
+		FRotation3 Orientation;
 		FVec3 Velocity;
 		FVec3 AngularVelocity;
 		FReal Mass;
