Make Playing (reading)  ChaosCaches a const operation.

This includes marking for read const by making the read lock mutable.

diff --git a/Engine/Plugins/Experimental/ChaosCaching/Source/ChaosCaching/Private/Chaos/Adapters/GeometryCollectionComponentCacheAdapter.cpp b/Engine/Plugins/Experimental/ChaosCaching/Source/ChaosCaching/Private/Chaos/Adapters/GeometryCollectionComponentCacheAdapter.cpp
--- a/Engine/Plugins/Experimental/ChaosCaching/Source/ChaosCaching/Private/Chaos/Adapters/GeometryCollectionComponentCacheAdapter.cpp
+++ b/Engine/Plugins/Experimental/ChaosCaching/Source/ChaosCaching/Private/Chaos/Adapters/GeometryCollectionComponentCacheAdapter.cpp
@@ -282,7 +282,7 @@ namespace Chaos
 			TMap<FClusterParticle*, TArray<FRigidParticle*>> NewClusters;
 			for(const FCacheEventHandle& Handle : *EnableEvents)
 			{
-				if(FEnableStateEvent* Event = Handle.Get<FEnableStateEvent>())
+				if( const FEnableStateEvent* Event = Handle.GetConst<FEnableStateEvent>())
 				{
 					if(Chaos::FPBDRigidClusteredParticleHandle* ChildParticle = Proxy->GetParticle_Internal(Event->Index))
 					{
@@ -347,7 +347,7 @@ namespace Chaos
 
 			for (const FCacheEventHandle& Handle : *BreakingEvents)
 			{
-				if (FBreakingEvent* Event = Handle.Get<FBreakingEvent>())
+				if (const FBreakingEvent* Event = Handle.GetConst<FBreakingEvent>())
 				{
 					if (Chaos::FPBDRigidClusteredParticleHandle* Particle = Proxy->GetParticle_Internal(Event->Index))
 					{
@@ -392,7 +392,7 @@ namespace Chaos
 
 			for (const FCacheEventHandle& Handle : *TrailingEvents)
 			{
-				if (FTrailingEvent* Event = Handle.Get<FTrailingEvent>())
+				if (const FTrailingEvent* Event = Handle.GetConst<FTrailingEvent>())
 				{
 					if (Chaos::FPBDRigidClusteredParticleHandle* Particle = Proxy->GetParticle_Internal(Event->Index))
 					{
@@ -435,7 +435,7 @@ namespace Chaos
 			const FSolverCollisionEventFilter* SolverCollisionEventFilter = Solver->GetEventFilters()->GetCollisionFilter();
 			for (const FCacheEventHandle& Handle : *CollisionEvents)
 			{
-				if (FCollisionEvent* Event = Handle.Get<FCollisionEvent>())
+				if (const FCollisionEvent* Event = Handle.GetConst<FCollisionEvent>())
 				{
 					FCollidingData CachedCollision;
 					CachedCollision.Location = ComponentToWorld.TransformPosition(Event->Location);
@@ -816,7 +816,7 @@ namespace Chaos
 			ReleaseIndices.Reserve(EnableEvents->Num());
 			for (const FCacheEventHandle& Handle : *EnableEvents)
 			{
-				if (FEnableStateEvent* Event = Handle.Get<FEnableStateEvent>())
+				if (const FEnableStateEvent* Event = Handle.GetConst<FEnableStateEvent>())
 				{
 					if (Event->bEnable)
 					{
diff --git a/Engine/Plugins/Experimental/ChaosCaching/Source/ChaosCaching/Private/Chaos/CacheEvents.cpp b/Engine/Plugins/Experimental/ChaosCaching/Source/ChaosCaching/Private/Chaos/CacheEvents.cpp
--- a/Engine/Plugins/Experimental/ChaosCaching/Source/ChaosCaching/Private/Chaos/CacheEvents.cpp
+++ b/Engine/Plugins/Experimental/ChaosCaching/Source/ChaosCaching/Private/Chaos/CacheEvents.cpp
@@ -24,7 +24,7 @@ FCacheEventTrack::~FCacheEventTrack()
 	DestroyAll();
 }
 
-FCacheEventTrack::FHandle FCacheEventTrack::GetEventHandle(int32 Index)
+FCacheEventTrack::FHandle FCacheEventTrack::GetEventHandle(int32 Index) const
 {
 	FHandle NewHandle;
 	if(EventData.IsValidIndex(Index))
diff --git a/Engine/Plugins/Experimental/ChaosCaching/Source/ChaosCaching/Private/Chaos/ChaosCache.cpp b/Engine/Plugins/Experimental/ChaosCaching/Source/ChaosCaching/Private/Chaos/ChaosCache.cpp
--- a/Engine/Plugins/Experimental/ChaosCaching/Source/ChaosCaching/Private/Chaos/ChaosCache.cpp
+++ b/Engine/Plugins/Experimental/ChaosCaching/Source/ChaosCaching/Private/Chaos/ChaosCache.cpp
@@ -495,7 +495,7 @@ void UChaosCache::CompressTracks()
 	}
 }
 
-FCacheUserToken UChaosCache::BeginPlayback()
+FCacheUserToken UChaosCache::BeginPlayback() const
 {
 	CurrentPlaybackCount++;
 	if(CurrentRecordCount.load() == 0)
@@ -511,7 +511,7 @@ FCacheUserToken UChaosCache::BeginPlayback()
 	return FCacheUserToken(false, false, this);
 }
 
-void UChaosCache::EndPlayback(FCacheUserToken& InOutToken)
+void UChaosCache::EndPlayback(FCacheUserToken& InOutToken) const
 {
 	if(InOutToken.IsOpen() && InOutToken.Owner == this)
 	{
@@ -543,7 +543,7 @@ float UChaosCache::GetDuration() const
 	return RecordedDuration;
 }
 
-FCacheEvaluationResult UChaosCache::Evaluate(const FCacheEvaluationContext& InContext, const TArray<FTransform>* MassToLocalTransforms)
+FCacheEvaluationResult UChaosCache::Evaluate(const FCacheEvaluationContext& InContext, const TArray<FTransform>* MassToLocalTransforms) const
 {
 	QUICK_SCOPE_CYCLE_COUNTER(QSTAT_CacheEval);
 
@@ -772,11 +772,11 @@ const FCacheSpawnableTemplate& UChaosCache::GetSpawnableTemplate() const
 	return Spawnable;
 }
 
-void UChaosCache::EvaluateSingle(int32 InIndex, FPlaybackTickRecord& InTickRecord, const FTransform* MassToLocal, FTransform* OutOptTransform, TMap<FName, float>* OutOptCurves)
+void UChaosCache::EvaluateSingle(int32 InIndex, FPlaybackTickRecord& InTickRecord, const FTransform* MassToLocal, FTransform* OutOptTransform, TMap<FName, float>* OutOptCurves) const
 {
 	// check to satisfy SA, external callers check validity in Evaluate
 	checkSlow(ParticleTracks.IsValidIndex(InIndex));
-	FPerParticleCacheData& Data = ParticleTracks[InIndex];
+	const FPerParticleCacheData& Data = ParticleTracks[InIndex];
 
 	
 	if(OutOptTransform)
@@ -792,12 +792,12 @@ void UChaosCache::EvaluateSingle(int32 InIndex, FPlaybackTickRecord& InTickRecor
 	}
 }
 
-void UChaosCache::EvaluateTransform(const FPerParticleCacheData& InData, float InTime, const FTransform* MassToLocal, FTransform& OutTransform)
+void UChaosCache::EvaluateTransform(const FPerParticleCacheData& InData, float InTime, const FTransform* MassToLocal, FTransform& OutTransform) const
 {
 	OutTransform = InData.TransformData.Evaluate(InTime, MassToLocal, InterpolationMode);
 }
 
-void UChaosCache::EvaluateCurves(const FPerParticleCacheData& InData, float InTime, TMap<FName, float>& OutCurves)
+void UChaosCache::EvaluateCurves(const FPerParticleCacheData& InData, float InTime, TMap<FName, float>& OutCurves) const
 {
 	for(const TPair<FName, FRichCurve>& Curve : InData.CurveData)
 	{
@@ -805,13 +805,13 @@ void UChaosCache::EvaluateCurves(const FPerParticleCacheData& InData, float InTi
 	}
 }
 
-void UChaosCache::EvaluateEvents(FPlaybackTickRecord& InTickRecord, TMap<FName, TArray<FCacheEventHandle>>& OutEvents)
+void UChaosCache::EvaluateEvents(FPlaybackTickRecord& InTickRecord, TMap<FName, TArray<FCacheEventHandle>>& OutEvents) const
 {
 	OutEvents.Reset();
 
-	for(TTuple<FName, FCacheEventTrack>& Track : EventTracks)
+	for(const TTuple<FName, FCacheEventTrack>& Track : EventTracks)
 	{
-		FCacheEventTrack& TrackRef = Track.Value;
+		const FCacheEventTrack& TrackRef = Track.Value;
 		if(TrackRef.Num() == 0)
 		{
 			continue;
@@ -820,7 +820,7 @@ void UChaosCache::EvaluateEvents(FPlaybackTickRecord& InTickRecord, TMap<FName,
 		int32* BeginIndexPtr = InTickRecord.LastEventPerTrack.Find(Track.Key);
 		const int32 BeginIndex = BeginIndexPtr ? *BeginIndexPtr : 0;
 
-		TArrayView<float> TimeStampView(&TrackRef.TimeStamps[BeginIndex], TrackRef.TimeStamps.Num() - BeginIndex);
+		TConstArrayView<float> TimeStampView(&TrackRef.TimeStamps[BeginIndex], TrackRef.TimeStamps.Num() - BeginIndex);
 
 		int32 BeginEventIndex = Algo::UpperBound(TimeStampView, InTickRecord.LastTime) + BeginIndex;
 		const int32 EndEventIndex = Algo::UpperBound(TimeStampView, InTickRecord.GetTime()) + BeginIndex;
diff --git a/Engine/Plugins/Experimental/ChaosCaching/Source/ChaosCaching/Public/Chaos/CacheEvents.h b/Engine/Plugins/Experimental/ChaosCaching/Source/ChaosCaching/Public/Chaos/CacheEvents.h
--- a/Engine/Plugins/Experimental/ChaosCaching/Source/ChaosCaching/Public/Chaos/CacheEvents.h
+++ b/Engine/Plugins/Experimental/ChaosCaching/Source/ChaosCaching/Public/Chaos/CacheEvents.h
@@ -35,7 +35,14 @@ struct FCacheEventTrack
 		}
 
 		template<typename T>
+		UE_DEPRECATED(5.6, "Non-const access is deprecated. Use GetConst instead")
 		T* Get() const
+		{
+			return IsAlive() ? const_cast<FCacheEventTrack*>(Track)->GetEvent<T>(Index) : nullptr;
+		}
+
+		template<typename T>
+		const T* GetConst() const
 		{
 			return IsAlive() ? Track->GetEvent<T>(Index) : nullptr;
 		}
@@ -46,7 +53,7 @@ struct FCacheEventTrack
 
 		bool IsAlive() const;
 
-		FCacheEventTrack* Track;
+		const FCacheEventTrack* Track;
 		int32             Version;
 		int32             Index;
 	};
@@ -95,8 +102,20 @@ struct FCacheEventTrack
 		return nullptr;
 	}
 
+	template<typename T>
+	const T* GetEvent(int32 Index) const
+	{
+		if (EventData.IsValidIndex(Index) && Struct)
+		{
+			check(T::StaticStruct() == Struct);
+			return reinterpret_cast<const T*>(EventData[Index]);
+		}
+
+		return nullptr;
+	}
+
 	/** Get a handle to an event that can easily resolve the inner event type without knowing the track */
-	FHandle GetEventHandle(int32 Index);
+	FHandle GetEventHandle(int32 Index) const;
 
 	template<typename T>
 	TArray<T*> GetEvents(float T0, float T1, int32* OptOutBeginIndex = nullptr)
diff --git a/Engine/Plugins/Experimental/ChaosCaching/Source/ChaosCaching/Public/Chaos/ChaosCache.h b/Engine/Plugins/Experimental/ChaosCaching/Source/ChaosCaching/Public/Chaos/ChaosCache.h
--- a/Engine/Plugins/Experimental/ChaosCaching/Source/ChaosCaching/Public/Chaos/ChaosCache.h
+++ b/Engine/Plugins/Experimental/ChaosCaching/Source/ChaosCaching/Public/Chaos/ChaosCache.h
@@ -281,9 +281,9 @@ private:
 
 	bool         bIsOpen;
 	bool         bIsRecord;
-	UChaosCache* Owner;
+	const UChaosCache* Owner;
 
-	explicit FCacheUserToken(bool bInOpen, bool bInRecord, UChaosCache* InOwner)
+	explicit FCacheUserToken(bool bInOpen, bool bInRecord, const UChaosCache* InOwner)
 		: bIsOpen(bInOpen)
 		, bIsRecord(bInRecord)
 		, Owner(InOwner)
@@ -349,7 +349,7 @@ public:
 	 * Initialise the cache for playback, may not take any actual action on the cache but
 	 * will provide the caller with a valid cache user token if it is safe to continue with playback
 	 */
-	CHAOSCACHING_API FCacheUserToken BeginPlayback();
+	CHAOSCACHING_API FCacheUserToken BeginPlayback() const;
 
 	/**
 	 * End a playback session for the cache. There can be multiple playback sessions open for a
@@ -357,7 +357,7 @@ public:
 	 * token will decrease the session count.
 	 * @param InOutToken The token that was given by BeginRecord
 	 */
-	CHAOSCACHING_API void EndPlayback(FCacheUserToken& InOutToken);
+	CHAOSCACHING_API void EndPlayback(FCacheUserToken& InOutToken) const;
 
 	/**
 	 * Adds a new frame to process to a threadsafe queue for later processing in FlushPendingFrames
@@ -376,7 +376,7 @@ public:
 	 * @param MassToLocalTransforms MassToLocal trasnform if available ( geometry collection are using them )
 	 * @see FCacheEvaluationContext
 	 */
-	CHAOSCACHING_API FCacheEvaluationResult Evaluate(const FCacheEvaluationContext& InContext, const TArray<FTransform>* MassToLocalTransforms);
+	CHAOSCACHING_API FCacheEvaluationResult Evaluate(const FCacheEvaluationContext& InContext, const TArray<FTransform>* MassToLocalTransforms) const;
 
 	/**
 	 * Initializes the spawnable template from a currently existing component so it can be spawned by the editor
@@ -404,11 +404,11 @@ public:
 	 * @param OutOptTransform Transform to fill, skipped if null
 	 * @param OutOptCurves Curves to fill, skipped if null
 	 */
-	CHAOSCACHING_API void EvaluateSingle(int32 InIndex, FPlaybackTickRecord& InTickRecord, const FTransform* MassToLocal, FTransform* OutOptTransform, TMap<FName, float>* OutOptCurves);
+	CHAOSCACHING_API void EvaluateSingle(int32 InIndex, FPlaybackTickRecord& InTickRecord, const FTransform* MassToLocal, FTransform* OutOptTransform, TMap<FName, float>* OutOptCurves) const;
 
-	CHAOSCACHING_API void EvaluateTransform(const FPerParticleCacheData& InData, float InTime, const FTransform* MassToLocal, FTransform& OutTransform);
-	CHAOSCACHING_API void EvaluateCurves(const FPerParticleCacheData& InData, float InTime, TMap<FName, float>& OutCurves);
-	CHAOSCACHING_API void EvaluateEvents(FPlaybackTickRecord& InTickRecord, TMap<FName, TArray<FCacheEventHandle>>& OutEvents);
+	CHAOSCACHING_API void EvaluateTransform(const FPerParticleCacheData& InData, float InTime, const FTransform* MassToLocal, FTransform& OutTransform) const;
+	CHAOSCACHING_API void EvaluateCurves(const FPerParticleCacheData& InData, float InTime, TMap<FName, float>& OutCurves) const;
+	CHAOSCACHING_API void EvaluateEvents(FPlaybackTickRecord& InTickRecord, TMap<FName, TArray<FCacheEventHandle>>& OutEvents) const;
 
 	CHAOSCACHING_API void CompressChannelsData(float ErrorThreshold, float SampleRate);
 
@@ -511,7 +511,7 @@ private:
 
 	/** Counts for current number of users, should only ever have one recorder, and if we do no playbacks */
 	std::atomic<int32> CurrentRecordCount;
-	std::atomic<int32> CurrentPlaybackCount;
+	mutable std::atomic<int32> CurrentPlaybackCount;
 
 	/** Indicates that we need to strip MassToLocal before playing the cache. */
 	bool bStripMassToLocal;
