Chaos: Fix overflow compilation error with FPSemanticsMode.Precise.

diff --git a/Engine/Plugins/Experimental/ChaosCaching/Source/ChaosCaching/Public/Chaos/Adapters/CacheAdapter.h b/Engine/Plugins/Experimental/ChaosCaching/Source/ChaosCaching/Public/Chaos/Adapters/CacheAdapter.h
--- a/Engine/Plugins/Experimental/ChaosCaching/Source/ChaosCaching/Public/Chaos/Adapters/CacheAdapter.h
+++ b/Engine/Plugins/Experimental/ChaosCaching/Source/ChaosCaching/Public/Chaos/Adapters/CacheAdapter.h
@@ -210,8 +210,8 @@ namespace Chaos
 
 	protected:
 		bool bRestartSimulation = false;
-		float RestartTimeStart = TNumericLimits<double>::Max(); // if bRestartSimulation, use RestartTimeStart as an offset when saving, and clear cache data in [RestartTimeStart, RestartTimeEnd] 
-		float RestartTimeEnd = -TNumericLimits<double>::Max();
+		float RestartTimeStart = TNumericLimits<float>::Max(); // if bRestartSimulation, use RestartTimeStart as an offset when saving, and clear cache data in [RestartTimeStart, RestartTimeEnd] 
+		float RestartTimeEnd = -TNumericLimits<float>::Max();
 
 	private:
 	};
diff --git a/Engine/Plugins/Experimental/ChaosCaching/Source/ChaosCaching/Public/Chaos/CacheManagerActor.h b/Engine/Plugins/Experimental/ChaosCaching/Source/ChaosCaching/Public/Chaos/CacheManagerActor.h
--- a/Engine/Plugins/Experimental/ChaosCaching/Source/ChaosCaching/Public/Chaos/CacheManagerActor.h
+++ b/Engine/Plugins/Experimental/ChaosCaching/Source/ChaosCaching/Public/Chaos/CacheManagerActor.h
@@ -400,8 +400,8 @@ private:
 	float StartTimeAtBeginPlay;
 
 	bool bRestartSimulation = false; // if restart for a specified time range
-	float RestartTimeStart = TNumericLimits<double>::Max(); // if bRestartSimulation, use RestartTimeStart as an offset when saving, and clear cache data in [RestartTimeStart, RestartTimeEnd] 
-	float RestartTimeEnd = -TNumericLimits<double>::Max();
+	float RestartTimeStart = TNumericLimits<float>::Max(); // if bRestartSimulation, use RestartTimeStart as an offset when saving, and clear cache data in [RestartTimeStart, RestartTimeEnd] 
+	float RestartTimeEnd = -TNumericLimits<float>::Max();
 
 #if WITH_EDITORONLY_DATA
 	UPROPERTY(Transient)
