Chaos cache : Add Start and Stop control to blueprint to allow blueprint to set start time when running in play mode

diff --git a/Engine/Plugins/Experimental/ChaosCaching/Source/ChaosCaching/Private/Chaos/CacheManagerActor.cpp b/Engine/Plugins/Experimental/ChaosCaching/Source/ChaosCaching/Private/Chaos/CacheManagerActor.cpp
--- a/Engine/Plugins/Experimental/ChaosCaching/Source/ChaosCaching/Private/Chaos/CacheManagerActor.cpp
+++ b/Engine/Plugins/Experimental/ChaosCaching/Source/ChaosCaching/Private/Chaos/CacheManagerActor.cpp
@@ -327,6 +327,22 @@ void AChaosCacheManager::SelectComponent(int32 InIndex)
 #endif
 
 
+void AChaosCacheManager::Stop()
+{
+	EndEvaluate();
+}
+
+void AChaosCacheManager::Start(float InStartTime)
+{
+	if (CacheMode != ECacheMode::Record)
+	{
+		Stop();
+		StartTime = InStartTime;
+		OnStartFrameChanged(InStartTime);
+		BeginEvaluate();
+	}
+}
+
 void AChaosCacheManager::BeginPlay()
 {
 	using namespace Chaos;
diff --git a/Engine/Plugins/Experimental/ChaosCaching/Source/ChaosCaching/Public/Chaos/CacheManagerActor.h b/Engine/Plugins/Experimental/ChaosCaching/Source/ChaosCaching/Public/Chaos/CacheManagerActor.h
--- a/Engine/Plugins/Experimental/ChaosCaching/Source/ChaosCaching/Public/Chaos/CacheManagerActor.h
+++ b/Engine/Plugins/Experimental/ChaosCaching/Source/ChaosCaching/Public/Chaos/CacheManagerActor.h
@@ -227,13 +227,34 @@ public:
 #endif
 	/** end UObject interface */
 
-	/** Expose StartTime property to Sequencer. GetStartTime will be called on keys. */
+	/** 
+	* Expose StartTime property to Sequencer. GetStartTime will be called on keys. 
+	* Note this will be ignored for evaluating in Play mode ( for that case, make sure to use Start(time) instead as this will properly stop the physics thread evalution and restert it after setting the time ) 
+	*/
 	UFUNCTION(CallInEditor)
 	CHAOSCACHING_API void SetStartTime(float InStartTime);
 
+	/**
+	* Used to set the start time of the cache from a blueprint 
+	* Note this will be ignored for evaluating in Play mode ( for that case, make sure to use Start(time) instead as this will properly stop the physics thread evalution and restert it after setting the time )
+	*/
 	UFUNCTION(BlueprintCallable, Category = "Caching")
 	void SetCurrentTime(float CurrentTime);
 
+	/* 
+	* Start the cache evaluation at a specific time 
+	* If the cache if already eveluating, stop it, set the time, then re-start 
+	* only works when the cache is in Static Pose or Play mode
+	*/
+	UFUNCTION(BlueprintCallable, Category = "Caching")
+	void Start(float InStartTime = 0);
+
+	/* 
+	* stop the cache from eveluating 
+	*/
+	UFUNCTION(BlueprintCallable, Category = "Caching")
+	void Stop();
+
 	/** 
 	 * Resets all components back to the world space transform they had when the cache for them was originally recorded
 	 * if one is available
