Chaos cache : Add an option to decide to start or not the cache manager on BeginPlay
true by default, this will start the cache in record / play mode automatically on beginPlay
if set to false the cache will not record or play anything until Start is called on it

diff --git a/Engine/Plugins/Experimental/ChaosCaching/Source/ChaosCaching/Private/Chaos/CacheManagerActor.cpp b/Engine/Plugins/Experimental/ChaosCaching/Source/ChaosCaching/Private/Chaos/CacheManagerActor.cpp
--- a/Engine/Plugins/Experimental/ChaosCaching/Source/ChaosCaching/Private/Chaos/CacheManagerActor.cpp
+++ b/Engine/Plugins/Experimental/ChaosCaching/Source/ChaosCaching/Private/Chaos/CacheManagerActor.cpp
@@ -73,6 +73,7 @@ AChaosCacheManager::AChaosCacheManager(const FObjectInitializer& ObjectInitializ
 	: Super(ObjectInitializer)
 	, CacheMode(ECacheMode::Record)
 	, StartMode(EStartMode::Timed)
+	, bStartOnBeginPlay(true)
 	, bCanRecord(true)
 	, bIsSimulating(false)
 	, StartTimeAtBeginPlay(0)
@@ -350,7 +351,10 @@ void AChaosCacheManager::BeginPlay()
 	Super::BeginPlay();
 
 	StartTimeAtBeginPlay = StartTime;
-	BeginEvaluate();
+	if (bStartOnBeginPlay)
+	{
+		BeginEvaluate();
+	}
 }
 
 void AChaosCacheManager::BeginEvaluate()
diff --git a/Engine/Plugins/Experimental/ChaosCaching/Source/ChaosCaching/Public/Chaos/CacheManagerActor.h b/Engine/Plugins/Experimental/ChaosCaching/Source/ChaosCaching/Public/Chaos/CacheManagerActor.h
--- a/Engine/Plugins/Experimental/ChaosCaching/Source/ChaosCaching/Public/Chaos/CacheManagerActor.h
+++ b/Engine/Plugins/Experimental/ChaosCaching/Source/ChaosCaching/Public/Chaos/CacheManagerActor.h
@@ -207,6 +207,11 @@ public:
 	UPROPERTY(EditAnywhere, Interp, BlueprintReadWrite, Category = "Caching", meta=(SequencerTrackClass="/Script/MovieSceneTracks.MovieSceneFloatTrack"))
 	float StartTime;
 
+	// if true, cache will start automatically on begin play ( playing or recording )
+	// if false, the function Start must be used to tsrat the cache again ( and can be stopped again using Stop )
+	UPROPERTY(EditAnywhere, Category = "Caching")
+	bool bStartOnBeginPlay;
+
 	/** AActor interface */
 	CHAOSCACHING_API void TickActor(float DeltaTime, enum ELevelTick TickType, FActorTickFunction& ThisTickFunction) override;
 	/** end AActor interface */
