Chaos cache : fix issues with static pose replay when root proxy are enables

diff --git a/Engine/Plugins/Experimental/ChaosCaching/Source/ChaosCaching/Private/Chaos/Adapters/GeometryCollectionComponentCacheAdapter.cpp b/Engine/Plugins/Experimental/ChaosCaching/Source/ChaosCaching/Private/Chaos/Adapters/GeometryCollectionComponentCacheAdapter.cpp
--- a/Engine/Plugins/Experimental/ChaosCaching/Source/ChaosCaching/Private/Chaos/Adapters/GeometryCollectionComponentCacheAdapter.cpp
+++ b/Engine/Plugins/Experimental/ChaosCaching/Source/ChaosCaching/Private/Chaos/Adapters/GeometryCollectionComponentCacheAdapter.cpp
@@ -656,15 +656,24 @@ namespace Chaos
 				// Set any broken clusters to identity so that they no longer influence child transforms.
 				TArray<int32> ReleaseIndices = GatherAllBreaksUpToTime(InCache, InTime);
 				const TManagedArray<int32>& Parent = GeometryCollection->Parent;
+				const int32 RootIndex = Comp->GetRootIndex();
+				bool bIsRootBroken = false;
 				for (int32 ReleaseIdx : ReleaseIndices)
 				{
-					int32 ClusterIdx = Parent[ReleaseIdx];
+					const int32 ClusterIdx = Parent[ReleaseIdx];
+					bIsRootBroken = (ClusterIdx == RootIndex);
 					if (ClusterIdx > INDEX_NONE)
 					{
 						RestTransforms[ClusterIdx] = FTransform::Identity;
 					}
 				}
 
+				if (bIsRootBroken != Comp->IsForceBrokenForCustomRenderer())
+				{
+					// this is the root bone broken we need to disable the root proxy
+					Comp->ForceBrokenForCustomRenderer(bIsRootBroken);
+				}
+
 				Comp->SetRestState(MoveTemp(RestTransforms));
 			}
 		}
@@ -681,6 +690,11 @@ namespace Chaos
 		{
 			return false;
 		}
+		if (Comp)
+		{
+			// make sure we have the root proxy enabled if any
+			Comp->ForceBrokenForCustomRenderer(false);
+		}
 
 		const FSimulationParameters& SimulationParameters = Proxy->GetSimParameters();
 
@@ -778,6 +792,10 @@ namespace Chaos
 						InitialTransforms[TransformIdx] = EvaluatedResult.Transform[CacheIdx];
 					}
 				}
+
+				// make sure we have the root proxy enabled if any
+				Comp->ForceBrokenForCustomRenderer(false);
+
 				Comp->SetInitialTransforms(InitialTransforms);
 					
 				TArray<int32> ReleaseIndices = GatherAllBreaksUpToTime(InCache, InTime);
diff --git a/Engine/Source/Runtime/Experimental/GeometryCollectionEngine/Public/GeometryCollection/GeometryCollectionComponent.h b/Engine/Source/Runtime/Experimental/GeometryCollectionEngine/Public/GeometryCollection/GeometryCollectionComponent.h
--- a/Engine/Source/Runtime/Experimental/GeometryCollectionEngine/Public/GeometryCollection/GeometryCollectionComponent.h
+++ b/Engine/Source/Runtime/Experimental/GeometryCollectionEngine/Public/GeometryCollection/GeometryCollectionComponent.h
@@ -1345,6 +1345,8 @@ public:
 	UFUNCTION(BlueprintCallable, Category = "Physics")
 	GEOMETRYCOLLECTIONENGINE_API void ForceBrokenForCustomRenderer(bool bForceBroken);
 
+	GEOMETRYCOLLECTIONENGINE_API bool IsForceBrokenForCustomRenderer() const { return bForceBrokenForCustomRenderer; }
+
 	/** 
 	* blueprint function to set a specific root proxy local transform 
 	* warning: when called from C++ and calling it on multiple root proxies it is recommended to use SetRootProxyLocalTransform instead and then call RefreshCustomRenderer
