Chaos Cache : fix static analysis issue

diff --git a/Engine/Plugins/Experimental/ChaosCaching/Source/ChaosCaching/Private/Chaos/Adapters/GeometryCollectionComponentCacheAdapter.cpp b/Engine/Plugins/Experimental/ChaosCaching/Source/ChaosCaching/Private/Chaos/Adapters/GeometryCollectionComponentCacheAdapter.cpp
--- a/Engine/Plugins/Experimental/ChaosCaching/Source/ChaosCaching/Private/Chaos/Adapters/GeometryCollectionComponentCacheAdapter.cpp
+++ b/Engine/Plugins/Experimental/ChaosCaching/Source/ChaosCaching/Private/Chaos/Adapters/GeometryCollectionComponentCacheAdapter.cpp
@@ -684,7 +684,7 @@ namespace Chaos
 		EnsureIsInGameThreadContext();
 		
 		UGeometryCollectionComponent*    Comp     = CastChecked<UGeometryCollectionComponent>(InComponent);
-		FGeometryCollectionPhysicsProxy* Proxy    = Comp->GetPhysicsProxy();
+		FGeometryCollectionPhysicsProxy* Proxy = Comp ? Comp->GetPhysicsProxy() : nullptr;
 
 		if (!Proxy)
 		{
