Chaos cache manager : fix rare factory crash
- make sure to use IsInstanceOf to make sure the class pointer is properly checked for null

diff --git a/Engine/Plugins/Experimental/ChaosCaching/Source/ChaosCachingEditor/Private/Chaos/ActorFactoryCacheManager.cpp b/Engine/Plugins/Experimental/ChaosCaching/Source/ChaosCachingEditor/Private/Chaos/ActorFactoryCacheManager.cpp
--- a/Engine/Plugins/Experimental/ChaosCaching/Source/ChaosCachingEditor/Private/Chaos/ActorFactoryCacheManager.cpp
+++ b/Engine/Plugins/Experimental/ChaosCaching/Source/ChaosCachingEditor/Private/Chaos/ActorFactoryCacheManager.cpp
@@ -21,7 +21,7 @@ UActorFactoryCacheManager::UActorFactoryCacheManager()
 
 bool UActorFactoryCacheManager::CanCreateActorFrom(const FAssetData& AssetData, FText& OutErrorMsg)
 {
-	if(!AssetData.IsValid() || !AssetData.GetClass()->IsChildOf(UChaosCacheCollection::StaticClass()))
+	if (!AssetData.IsValid() || !AssetData.IsInstanceOf(UChaosCacheCollection::StaticClass()))
 	{
 		OutErrorMsg = LOCTEXT("NoCollection", "A valid cache collection must be specified.");
 		return false;
