Chaos cache: fix issues with events not recording properly after recording once
- the issue was caused by multiple things:
- The proxies were never removed from the CachedData, accumulating over each recording session
- the event registration logic was assuming that all proxies across all solvers had the same event notification parameters
- Added a FinalizeComponent method to adapaters so they have the opprtunity to done per component / proxy cleanup
- Track registration per Solver/Event pair to make sure we only register event handlers that we need
- Also fixed mass not being initialized set on replay of trailing events as well as the TransformGroup Index

diff --git a/Engine/Plugins/Experimental/ChaosCaching/Source/ChaosCaching/Private/Chaos/Adapters/GeometryCollectionComponentCacheAdapter.cpp b/Engine/Plugins/Experimental/ChaosCaching/Source/ChaosCaching/Private/Chaos/Adapters/GeometryCollectionComponentCacheAdapter.cpp
--- a/Engine/Plugins/Experimental/ChaosCaching/Source/ChaosCaching/Private/Chaos/Adapters/GeometryCollectionComponentCacheAdapter.cpp
+++ b/Engine/Plugins/Experimental/ChaosCaching/Source/ChaosCaching/Private/Chaos/Adapters/GeometryCollectionComponentCacheAdapter.cpp
@@ -410,6 +410,8 @@ namespace Chaos
 						CachedTrail.AngularVelocity = Event->AngularVelocity;
 						CachedTrail.BoundingBox = TAABB<FReal, 3>(Event->BoundingBoxMin, Event->BoundingBoxMax);
 						CachedTrail.BoundingBox = CachedTrail.BoundingBox.TransformedAABB(ComponentToWorld);
+						CachedTrail.Mass = Particle->M();
+						CachedTrail.TransformGroupIndex = INDEX_NONE; // todo(ccaillaud) we should probably record that info in the cache : but we need to be careful that this value may change with the GC
 
 						if (!SolverTrailingEventFilter->Enabled() || SolverTrailingEventFilter->Pass(CachedTrail))
 						{
@@ -679,6 +681,110 @@ namespace Chaos
 		}
 	}
 
+	void FGeometryCollectionCacheAdapter::RegisterEventHandler(FGeometryCollectionPhysicsProxy* Proxy)
+	{
+		if (!Proxy)
+		{
+			return;
+		}
+
+		if (Chaos::FPhysicsSolver* Solver = Proxy->GetSolver<Chaos::FPhysicsSolver>())
+		{
+			const FSimulationParameters& SimulationParameters = Proxy->GetSimParameters();
+
+			// this method ensure that only the necessary handlers are registered by solver and for the need of the specific proxies we are tracking
+			if (Chaos::FEventManager* EventManager = Solver->GetEventManager())
+			{
+				if (SimulationParameters.bGenerateCollisionData || SimulationParameters.bGenerateGlobalCollisionData)
+				{
+					const FSolverEventPair SolverCollisionEventPair{ Solver,Chaos::EEventType::Collision };
+					if (!SolversListenedToForEvents.Contains(SolverCollisionEventPair))
+					{
+						EventManager->RegisterHandler<Chaos::FCollisionEventData>(Chaos::EEventType::Collision, this, &FGeometryCollectionCacheAdapter::HandleCollisionEvents);
+						SolversListenedToForEvents.Add(SolverCollisionEventPair);
+					}
+				}
+				if (SimulationParameters.bGenerateBreakingData || SimulationParameters.bGenerateGlobalBreakingData)
+				{
+					const FSolverEventPair SolverBreakingEventPair{ Solver,Chaos::EEventType::Breaking };
+					if (!SolversListenedToForEvents.Contains(SolverBreakingEventPair))
+					{
+						EventManager->RegisterHandler<Chaos::FBreakingEventData>(Chaos::EEventType::Breaking, this, &FGeometryCollectionCacheAdapter::HandleBreakingEvents);
+						SolversListenedToForEvents.Add(SolverBreakingEventPair);
+					}
+				}
+				if (SimulationParameters.bGenerateTrailingData)
+				{
+					const FSolverEventPair SolverTrailingEventPair{ Solver,Chaos::EEventType::Trailing};
+					if (!SolversListenedToForEvents.Contains(SolverTrailingEventPair))
+					{
+						EventManager->RegisterHandler<Chaos::FTrailingEventData>(Chaos::EEventType::Trailing, this, &FGeometryCollectionCacheAdapter::HandleTrailingEvents);
+						SolversListenedToForEvents.Add(SolverTrailingEventPair);
+					}
+				}
+			}
+		}
+	}
+
+	void FGeometryCollectionCacheAdapter::UnregisterEventHandler(FGeometryCollectionPhysicsProxy* Proxy)
+	{
+		if (!Proxy)
+		{
+			return;
+		}
+
+		if (Chaos::FPhysicsSolver* Solver = Proxy->GetSolver<Chaos::FPhysicsSolver>())
+		{
+			// find all proxies for this Solver and check what they need to listen to after removal of the current proxy
+			bool bListenToCollisionEvents = false;
+			bool bListenToBreakingEvents = false;
+			bool bListenToTrailingEvents = false;
+			for (const TPair<IPhysicsProxyBase*, FCachedEventData>& CacheDataPair : CachedData)
+			{
+				const IPhysicsProxyBase* OtherProxy = CacheDataPair.Key;
+				const FGeometryCollectionPhysicsProxy* OtherGCProxy =
+					OtherProxy->GetType() == FGeometryCollectionPhysicsProxy::ConcreteType()
+					? static_cast<const FGeometryCollectionPhysicsProxy*>(OtherProxy)
+					: nullptr;
+
+				// exclude the proxy we are about to remove 
+				if (Proxy != OtherGCProxy && OtherGCProxy)
+				{
+					const FSimulationParameters& OtherSimParams = OtherGCProxy->GetSimParameters();
+					if (OtherProxy->GetSolver<Chaos::FPhysicsSolver>() == Solver)
+					{
+						bListenToCollisionEvents |= (OtherSimParams.bGenerateCollisionData || OtherSimParams.bGenerateGlobalCollisionData);
+						bListenToBreakingEvents |= (OtherSimParams.bGenerateBreakingData || OtherSimParams.bGenerateGlobalBreakingData);
+						bListenToTrailingEvents |= (OtherSimParams.bGenerateTrailingData);
+					}
+				}
+			}
+
+			// unregister the handlers that we do not need anymore
+			if (Chaos::FEventManager* EventManager = Solver->GetEventManager())
+			{
+				if (!bListenToCollisionEvents)
+				{
+					const FSolverEventPair SolverCollisionEventPair{ Solver,Chaos::EEventType::Collision };
+					EventManager->UnregisterHandler(Chaos::EEventType::Collision, this);
+					SolversListenedToForEvents.Remove(SolverCollisionEventPair);
+				}
+				if (!bListenToBreakingEvents)
+				{
+					const FSolverEventPair SolverBreakingEventPair{ Solver,Chaos::EEventType::Breaking };
+					EventManager->UnregisterHandler(Chaos::EEventType::Breaking, this);
+					SolversListenedToForEvents.Remove(SolverBreakingEventPair);
+				}
+				if (!bListenToTrailingEvents)
+				{
+					const FSolverEventPair SolverTrailingEventPair{ Solver,Chaos::EEventType::Trailing };
+					EventManager->UnregisterHandler(Chaos::EEventType::Trailing, this);
+					SolversListenedToForEvents.Remove(SolverTrailingEventPair);
+				}
+			}
+		}
+	}
+
 	bool FGeometryCollectionCacheAdapter::InitializeForRecord(UPrimitiveComponent* InComponent, UChaosCache* InCache)
 	{
 		EnsureIsInGameThreadContext();
@@ -696,10 +802,7 @@ namespace Chaos
 			Comp->ForceBrokenForCustomRenderer(false);
 		}
 
-		const FSimulationParameters& SimulationParameters = Proxy->GetSimParameters();
-
 		Chaos::FPhysicsSolver* Solver = Proxy->GetSolver<Chaos::FPhysicsSolver>();
-
 		if (!Solver)
 		{
 			return false;
@@ -718,28 +821,11 @@ namespace Chaos
 			});
 
 		
-		// We only need to register event handlers once, the first time we initialize.
+		RegisterEventHandler(Proxy);
+
+		// we only need to clear the data for the first proxy we register
 		if (CachedData.Num() == 0)
 		{
-			Chaos::FEventManager* EventManager = Solver->GetEventManager();
-			if (EventManager)
-			{
-				if (SimulationParameters.bGenerateBreakingData || SimulationParameters.bGenerateGlobalBreakingData)
-				{ 
-					EventManager->RegisterHandler<Chaos::FBreakingEventData>(Chaos::EEventType::Breaking, const_cast<FGeometryCollectionCacheAdapter*>(this), &FGeometryCollectionCacheAdapter::HandleBreakingEvents);
-				}
-
-				if (SimulationParameters.bGenerateCollisionData || SimulationParameters.bGenerateGlobalCollisionData)
-				{ 
-					EventManager->RegisterHandler<Chaos::FCollisionEventData>(Chaos::EEventType::Collision, const_cast<FGeometryCollectionCacheAdapter*>(this), &FGeometryCollectionCacheAdapter::HandleCollisionEvents);
-				}
-
-				if (SimulationParameters.bGenerateTrailingData)
-				{ 
-					EventManager->RegisterHandler<Chaos::FTrailingEventData>(Chaos::EEventType::Trailing, const_cast<FGeometryCollectionCacheAdapter*>(this), &FGeometryCollectionCacheAdapter::HandleTrailingEvents);
-				}
-			}
-
 			BreakingDataArray = nullptr;
 			CollisionDataArray = nullptr;
 			TrailingDataArray = nullptr;
@@ -813,6 +899,21 @@ namespace Chaos
 		return false;
 	}
 
+	void FGeometryCollectionCacheAdapter::FinalizeComponent(UPrimitiveComponent* InComponent, UChaosCache* InCache)
+	{
+		EnsureIsInGameThreadContext();
+
+		// make sure we remove the proxy from the list to avoid accumulating them record after record
+		if (UGeometryCollectionComponent* Comp = CastChecked<UGeometryCollectionComponent>(InComponent))
+		{
+			if (FGeometryCollectionPhysicsProxy* Proxy = Comp->GetPhysicsProxy())
+			{
+				UnregisterEventHandler(Proxy);
+				CachedData.Remove(Proxy);
+			}
+		}
+	}
+
 	TArray<int32> FGeometryCollectionCacheAdapter::GatherAllBreaksUpToTime(UChaosCache* InCache, float InTime) const
 	{
 		// Evaluate all breaking event that have occurred from the beginning of the cache up to the specified time.
diff --git a/Engine/Plugins/Experimental/ChaosCaching/Source/ChaosCaching/Private/Chaos/CacheManagerActor.cpp b/Engine/Plugins/Experimental/ChaosCaching/Source/ChaosCaching/Private/Chaos/CacheManagerActor.cpp
--- a/Engine/Plugins/Experimental/ChaosCaching/Source/ChaosCaching/Private/Chaos/CacheManagerActor.cpp
+++ b/Engine/Plugins/Experimental/ChaosCaching/Source/ChaosCaching/Private/Chaos/CacheManagerActor.cpp
@@ -591,6 +591,18 @@ void AChaosCacheManager::EndEvaluate()
 
 	WaitForObservedComponentSolverTasks();
 
+	check(ActiveAdapters.IsEmpty() || ObservedComponents.Num() == ActiveAdapters.Num());
+
+	// finalize all components
+	for (int32 CompIndex = 0; CompIndex < ActiveAdapters.Num(); ++CompIndex)
+	{
+		if (Chaos::FComponentCacheAdapter* Adapter = ActiveAdapters[CompIndex])
+		{
+			FObservedComponent& Observed = ObservedComponents[CompIndex];
+			Adapter->FinalizeComponent(Observed.GetComponent(this), Observed.GetChaosCache());
+		}
+	}
+
 	// Build list of adapters for our observed components
 	IModularFeatures& ModularFeatures = IModularFeatures::Get();
 	TArray<Chaos::FComponentCacheAdapter*> Adapters = ModularFeatures.GetModularFeatureImplementations<Chaos::FComponentCacheAdapter>(Chaos::FComponentCacheAdapter::FeatureName);
diff --git a/Engine/Plugins/Experimental/ChaosCaching/Source/ChaosCaching/Public/Chaos/Adapters/CacheAdapter.h b/Engine/Plugins/Experimental/ChaosCaching/Source/ChaosCaching/Public/Chaos/Adapters/CacheAdapter.h
--- a/Engine/Plugins/Experimental/ChaosCaching/Source/ChaosCaching/Public/Chaos/Adapters/CacheAdapter.h
+++ b/Engine/Plugins/Experimental/ChaosCaching/Source/ChaosCaching/Public/Chaos/Adapters/CacheAdapter.h
@@ -155,6 +155,11 @@ namespace Chaos
 		virtual bool InitializeForPlayback(UPrimitiveComponent* InComponent, UChaosCache* InCache, float InTime) { unimplemented(); return false; }
 		virtual bool CHAOSCACHING_API InitializeForPlayback(UPrimitiveComponent* InComponent, FObservedComponent& InObserved, float InTime);
 
+		/**
+		* Called from the game thread to finalize for a specific component
+		*/
+		virtual void FinalizeComponent(UPrimitiveComponent* InComponent, UChaosCache* InCache) {}
+
 		/**
 		 * Called by a cache observer when a component should be recorded to a cache through this adapter.
 		 * The time provided is the absolute time from the beginning of the cache recording that the adapter is
diff --git a/Engine/Plugins/Experimental/ChaosCaching/Source/ChaosCaching/Public/Chaos/Adapters/GeometryCollectionComponentCacheAdapter.h b/Engine/Plugins/Experimental/ChaosCaching/Source/ChaosCaching/Public/Chaos/Adapters/GeometryCollectionComponentCacheAdapter.h
--- a/Engine/Plugins/Experimental/ChaosCaching/Source/ChaosCaching/Public/Chaos/Adapters/GeometryCollectionComponentCacheAdapter.h
+++ b/Engine/Plugins/Experimental/ChaosCaching/Source/ChaosCaching/Public/Chaos/Adapters/GeometryCollectionComponentCacheAdapter.h
@@ -225,6 +225,9 @@ class UPrimitiveComponent;
 
 namespace Chaos
 {
+	// forward declaration to avoid adding EventManager.h
+	enum class EEventType : int32;
+
 	struct FCachedEventData
 	{
 		FCachedEventData()
@@ -255,6 +258,7 @@ namespace Chaos
 		virtual void				   SetRestState(UPrimitiveComponent* InComponent, UChaosCache* InCache, const FTransform& InRootTransform, Chaos::FReal InTime) const override;
 		virtual bool                   InitializeForRecord(UPrimitiveComponent* InComponent, UChaosCache* InCache) override;
 		virtual bool                   InitializeForPlayback(UPrimitiveComponent* InComponent, UChaosCache* InCache, float InTime) override;
+		virtual void				   FinalizeComponent(UPrimitiveComponent* InComponent, UChaosCache* InCache) override;
 		virtual void                   Record_PostSolve(UPrimitiveComponent* InComp, const FTransform& InRootTransform, FPendingFrameWrite& OutFrame, Chaos::FReal InTime) const override;
 		virtual void                   Playback_PreSolve(UPrimitiveComponent*								InComponent,
 												 UChaosCache*										InCache,
@@ -265,14 +269,24 @@ namespace Chaos
 		// End FComponentCacheAdapter interface
 
 	protected:
+
+		void RegisterEventHandler(FGeometryCollectionPhysicsProxy* Proxy);
+		void UnregisterEventHandler(FGeometryCollectionPhysicsProxy* Proxy);
 		void HandleBreakingEvents(const Chaos::FBreakingEventData& Event);
 		void HandleCollisionEvents(const Chaos::FCollisionEventData& Event);
 		void HandleTrailingEvents(const Chaos::FTrailingEventData& Event);
+
 	
 	private:
 		TArray<int32> GatherAllBreaksUpToTime(UChaosCache* InCache, float InTime) const;
 
-		
+		/** 
+		* List of the solvers we are currently listening for events for recording
+		* we need to track them because each of them owns a seperate unique event manager 
+		*/
+		using FSolverEventPair = TPair<Chaos::FPhysicsSolver*, Chaos::EEventType>;
+		TSet<FSolverEventPair> SolversListenedToForEvents;
+
 		TMap<IPhysicsProxyBase*, FCachedEventData> CachedData;
 		const FBreakingDataArray* BreakingDataArray;
 		const FCollisionDataArray* CollisionDataArray;
