Chaos Cache manager : expose observed component function to blueprint
- ability to add, update, remove and clear

diff --git a/Engine/Plugins/Experimental/ChaosCaching/Source/ChaosCaching/Public/Chaos/CacheManagerActor.h b/Engine/Plugins/Experimental/ChaosCaching/Source/ChaosCaching/Public/Chaos/CacheManagerActor.h
--- a/Engine/Plugins/Experimental/ChaosCaching/Source/ChaosCaching/Public/Chaos/CacheManagerActor.h
+++ b/Engine/Plugins/Experimental/ChaosCaching/Source/ChaosCaching/Public/Chaos/CacheManagerActor.h
@@ -297,14 +297,17 @@ public:
 
 	/** Accessor to the manager observed components (read/write) */
 	TArray<FObservedComponent>& GetObservedComponents() {return ObservedComponents;}
-	
+
 	/** Find or add a primitive component to a cache manager */
+	UFUNCTION(BlueprintCallable, Category = "Caching", meta = (DisplayName = "UpdateOrAddObservedComponent"))
 	CHAOSCACHING_API void FindOrAddObservedComponent(UPrimitiveComponent* InComponent, const FName& CacheName = TEXT(""), const bool bTransferSimulationFlag = false);
 
 	/** Remove a primitive component from the cache manager */
+	UFUNCTION(BlueprintCallable, Category = "Caching")
 	CHAOSCACHING_API void RemoveObservedComponent(UPrimitiveComponent* InComponent);
 	
 	/** Clear all the observed components */
+	UFUNCTION(BlueprintCallable, Category = "Caching")
 	CHAOSCACHING_API void ClearObservedComponents();
 	
 #if WITH_EDITOR
