Chaos cache manager : prevent transient component from being added when creating a chaos actor manager
- this avoid getting the root proxy static mesh component to be added wrongly

diff --git a/Engine/Plugins/Experimental/ChaosCaching/Source/ChaosCaching/Private/Chaos/CacheManagerActor.cpp b/Engine/Plugins/Experimental/ChaosCaching/Source/ChaosCaching/Private/Chaos/CacheManagerActor.cpp
--- a/Engine/Plugins/Experimental/ChaosCaching/Source/ChaosCaching/Private/Chaos/CacheManagerActor.cpp
+++ b/Engine/Plugins/Experimental/ChaosCaching/Source/ChaosCaching/Private/Chaos/CacheManagerActor.cpp
@@ -946,6 +946,11 @@ void AChaosCacheManager::FindOrAddObservedComponent(UPrimitiveComponent* InCompo
 {
 	if(InComponent != nullptr)
 	{
+		// we do not want to record transient component as they may not exists when we start recording 
+		if (InComponent->HasAnyFlags(RF_Transient))
+		{
+			return;
+		}
 		FObservedComponent* FoundComponent = FindObservedComponent(InComponent);
 		if(FObservedComponent* ObservedComponent = FoundComponent ? FoundComponent : &AddNewObservedComponent(InComponent))
 		{
diff --git a/Engine/Plugins/Experimental/ChaosCaching/Source/ChaosCachingEditor/Private/Chaos/ChaosCachingEditorPlugin.cpp b/Engine/Plugins/Experimental/ChaosCaching/Source/ChaosCachingEditor/Private/Chaos/ChaosCachingEditorPlugin.cpp
--- a/Engine/Plugins/Experimental/ChaosCaching/Source/ChaosCachingEditor/Private/Chaos/ChaosCachingEditorPlugin.cpp
+++ b/Engine/Plugins/Experimental/ChaosCaching/Source/ChaosCachingEditor/Private/Chaos/ChaosCachingEditorPlugin.cpp
@@ -340,6 +340,12 @@ void IChaosCachingEditorPlugin::OnCreateCacheManager()
 
 				check(Manager);
 
+				// we do not want to record transient component as they may not exists when we start recording 
+				if (PrimitiveComp->HasAnyFlags(RF_Transient))
+				{
+					continue;
+				}
+
 				const FObservedComponent* Existing = Manager->FindObservedComponent(PrimitiveComp);
 
 				if(!Existing)
