Chaos Cache - Add the ability to incrementally create new cache asset each time we are recording
- re-arranged some part of the code to make sure we get consistet logic when switching time, cache mode or cache collection asset
- Also added IsRecording method accessible from blueprint

diff --git a/Engine/Plugins/Experimental/ChaosCaching/Source/ChaosCaching/Private/Chaos/CacheManagerActor.cpp b/Engine/Plugins/Experimental/ChaosCaching/Source/ChaosCaching/Private/Chaos/CacheManagerActor.cpp
--- a/Engine/Plugins/Experimental/ChaosCaching/Source/ChaosCaching/Private/Chaos/CacheManagerActor.cpp
+++ b/Engine/Plugins/Experimental/ChaosCaching/Source/ChaosCaching/Private/Chaos/CacheManagerActor.cpp
@@ -17,6 +17,7 @@
 #include UE_INLINE_GENERATED_CPP_BY_NAME(CacheManagerActor)
 
 #if WITH_EDITOR
+#include "AssetToolsModule.h"
 #include "Editor.h"
 #include "Framework/Notifications/NotificationManager.h"
 #include "UObject/ConstructorHelpers.h"
@@ -209,29 +210,15 @@ void AChaosCacheManager::PostEditChangeProperty(FPropertyChangedEvent& PropertyC
 
 	if (PropertyName == GET_MEMBER_NAME_CHECKED(AChaosCacheManager, CacheMode))
 	{
-		if ((CacheMode == ECacheMode::Record) && !CanRecord())
-		{
-			// If we're not allowed to record but have somehow been set to record mode, revert to default
-			CacheMode = ECacheMode::None;
-		}
-
-		SetObservedComponentProperties(CacheMode);
-		
-		if (CacheMode != ECacheMode::Record)
-		{
-			OnStartFrameChanged(StartTime);
-		}
-		else
-		{
-			OnStartFrameChanged(0.0);
-		}
+		SetCacheMode(CacheMode);
 	}
 	else if (PropertyName == GET_MEMBER_NAME_CHECKED(AChaosCacheManager, StartTime)) 
 	{
-		if (CacheMode != ECacheMode::Record)
-		{
-			OnStartFrameChanged(StartTime);
-		}
+		SetStartTime(StartTime);
+	}
+	else if (PropertyName == GET_MEMBER_NAME_CHECKED(AChaosCacheManager, CacheCollection))
+	{
+		SetCacheCollection(CacheCollection);
 	}
 }
 #endif
@@ -253,10 +240,34 @@ void AChaosCacheManager::Serialize(FArchive& Ar)
 	}
 }
 
+void AChaosCacheManager::SetCacheMode(ECacheMode NewCacheMode)
+{
+	CacheMode = NewCacheMode;
+	if (IsRecording() && !CanRecord())
+	{
+		// If we're not allowed to record but have somehow been set to record mode, revert to default
+		CacheMode = ECacheMode::None;
+	}
+
+#if WITH_EDITOR
+	SetObservedComponentProperties(CacheMode);
+#endif
+
+	SetStartTime(IsRecording() ? 0 : StartTime);
+}
+
+bool AChaosCacheManager::IsRecording() const
+{
+	return (CacheMode == ECacheMode::Record);
+}
+
 void AChaosCacheManager::SetStartTime(float InStartTime)
 {
 	StartTime = InStartTime;
-	OnStartFrameChanged(InStartTime);
+	if (!IsRecording())
+	{
+		OnStartFrameChanged(InStartTime);
+	}
 }
 
 void AChaosCacheManager::SetCurrentTime(float CurrentTime)
@@ -335,13 +346,19 @@ void AChaosCacheManager::Stop()
 
 void AChaosCacheManager::Start(float InStartTime)
 {
-	if (CacheMode != ECacheMode::Record)
+	Stop();
+
+	if (IsRecording() && bCreateNewCacheCollectionAssetOnRecord)
+	{
+		CreateAndSetNewCacheAsset();
+	}
+
+	if (!IsRecording())
 	{
-		Stop();
-		StartTime = InStartTime;
+		StartTime = InStartTime; // todo(ccaillaud) - is that needed in record ?
 		OnStartFrameChanged(InStartTime);
-		BeginEvaluate();
 	}
+	BeginEvaluate();
 }
 
 void AChaosCacheManager::BeginPlay()
@@ -353,7 +370,7 @@ void AChaosCacheManager::BeginPlay()
 	StartTimeAtBeginPlay = StartTime;
 	if (bStartOnBeginPlay)
 	{
-		BeginEvaluate();
+		Start(StartTime);
 	}
 }
 
@@ -563,7 +580,7 @@ void AChaosCacheManager::BeginEvaluate()
 	}
 	else
 	{
-		if (CacheMode != ECacheMode::Record)
+		if (!IsRecording())
 		{
 			OnStartFrameChanged(StartTime);
 		}
@@ -580,6 +597,21 @@ void AChaosCacheManager::EndPlay(const EEndPlayReason::Type EndPlayReason)
 	{
 		SetStartTime(StartTimeAtBeginPlay);
 	}
+
+#if WITH_EDITOR
+	// if the cache collection was changed during PIE ( automatic asset creation upon recording for example ) 
+	// make sure we update the editor actor with this new collection so it canb be replayed properly in non-PIE mode
+	if (AActor* EditorWorldActor = EditorUtilities::GetEditorWorldCounterpartActor(this))
+	{
+		if (AChaosCacheManager* EditCachemanagerActor = Cast<AChaosCacheManager>(EditorWorldActor))
+		{
+			if (EditCachemanagerActor->CacheCollection != CacheCollection)
+			{
+				EditCachemanagerActor->SetCacheCollection(CacheCollection);
+			}
+		}
+	}
+#endif
 }
 
 void AChaosCacheManager::EndEvaluate()
@@ -847,6 +879,39 @@ void AChaosCacheManager::SetCacheCollection(UChaosCacheCollection* InCacheCollec
 	}
 
 	CacheCollection = InCacheCollection;
+
+	// Reset cache collection
+	for (FObservedComponent& Observed : ObservedComponents)
+	{
+		if (Observed.Cache && Observed.Cache->GetOuter() != CacheCollection)
+		{
+			Observed.Cache = nullptr;
+		}
+	}
+
+	// reupdate the pose based on the new cache 
+	SetStartTime(0);
+}
+
+void AChaosCacheManager::CreateAndSetNewCacheAsset()
+{
+#if WITH_EDITOR
+	if (CacheCollection)
+	{
+		FAssetToolsModule& AssetToolsModule = FModuleManager::GetModuleChecked<FAssetToolsModule>("AssetTools");
+
+		FString NewPackageName;
+		FString NewAssetName;
+
+		static const FString EmptyDefaultSuffix;
+		AssetToolsModule.Get().CreateUniqueAssetName(CacheCollection->GetOutermost()->GetName(), EmptyDefaultSuffix, /*out*/ NewPackageName, /*out*/ NewAssetName);
+		const FString NewPackagePath = FPackageName::GetLongPackagePath(NewPackageName);
+		UObject* NewObject = AssetToolsModule.Get().CreateAsset(NewAssetName, NewPackagePath, UChaosCacheCollection::StaticClass(), /* Factory */ nullptr);
+		CacheCollection = Cast<UChaosCacheCollection>(NewObject);
+	}
+#else
+	UE_LOG(LogChaosCache, Error, TEXT("Cache Manager [%s] - Creating new asset upon recording is not supported in non-editor builds"), *GetName());
+#endif
 }
 
 void AChaosCacheManager::TriggerComponent(UPrimitiveComponent* InComponent)
diff --git a/Engine/Plugins/Experimental/ChaosCaching/Source/ChaosCaching/Public/Chaos/CacheManagerActor.h b/Engine/Plugins/Experimental/ChaosCaching/Source/ChaosCaching/Public/Chaos/CacheManagerActor.h
--- a/Engine/Plugins/Experimental/ChaosCaching/Source/ChaosCaching/Public/Chaos/CacheManagerActor.h
+++ b/Engine/Plugins/Experimental/ChaosCaching/Source/ChaosCaching/Public/Chaos/CacheManagerActor.h
@@ -188,6 +188,12 @@ public:
 	UPROPERTY(EditInstanceOnly, Category = "Caching")
 	TObjectPtr<UChaosCacheCollection> CacheCollection;
 
+	/**
+	* When on, each time the cache starts recording a new asset will be created as an increment of the previous one
+	*/
+	UPROPERTY(EditAnywhere, Category = "Caching", meta = (EditCondition = "CacheCollection != nullptr"))
+	bool bCreateNewCacheCollectionAssetOnRecord = false;
+
 	/** How to use the cache - playback or record */
 	UPROPERTY(EditAnywhere, Category = "Caching")
 	ECacheMode CacheMode;
@@ -232,6 +238,11 @@ public:
 #endif
 	/** end UObject interface */
 
+	CHAOSCACHING_API void SetCacheMode(ECacheMode NewCacheMode);
+
+	UFUNCTION(BlueprintCallable, Category = "Caching")
+	CHAOSCACHING_API bool IsRecording() const;
+
 	/** 
 	* Expose StartTime property to Sequencer. GetStartTime will be called on keys. 
 	* Note this will be ignored for evaluating in Play mode ( for that case, make sure to use Start(time) instead as this will properly stop the physics thread evalution and restert it after setting the time ) 
@@ -411,6 +422,12 @@ private:
 	*/
 	CHAOSCACHING_API void WaitForObservedComponentSolverTasks();
 
+	/**
+	* this creates a new cache collection asset based on the existing one and set it as the current 
+	* this also find the editor counterpart actor and update the cache collection accordingly
+	*/
+	void CreateAndSetNewCacheAsset();
+
 	/** List of observed objects and their caches */
 	UPROPERTY(EditAnywhere, Category = "Caching")
 	TArray<FObservedComponent> ObservedComponents;
