Chaos Caching : make creation of chaos cache manager undo-able ( from the Actor > Chaos  menu )

diff --git a/Engine/Plugins/Experimental/ChaosCaching/Source/ChaosCachingEditor/Private/Chaos/ChaosCachingEditorPlugin.cpp b/Engine/Plugins/Experimental/ChaosCaching/Source/ChaosCachingEditor/Private/Chaos/ChaosCachingEditorPlugin.cpp
--- a/Engine/Plugins/Experimental/ChaosCaching/Source/ChaosCachingEditor/Private/Chaos/ChaosCachingEditorPlugin.cpp
+++ b/Engine/Plugins/Experimental/ChaosCaching/Source/ChaosCachingEditor/Private/Chaos/ChaosCachingEditorPlugin.cpp
@@ -292,6 +292,8 @@ void IChaosCachingEditorPlugin::RegisterCachingSubMenu(FMenuBuilder& InMenuBuild
 
 void IChaosCachingEditorPlugin::OnCreateCacheManager()
 {
+	FScopedTransaction Transaction(LOCTEXT("ChaosCachingEditorPlugin_CreateCacheManager_Transaction", "Create a Chaos Cache Manager From Selection"));
+
 	auto SpawnManager = [](UWorld* InWorld) -> AChaosCacheManager* {
 		FActorSpawnParameters Params;
 		return InWorld->SpawnActor<AChaosCacheManager>();
