Chaos cache : restore simulate physics and event flags when deleting the cache actor or when the observed component is changed modified
- also make sure we store the original status of the break event flag on geometry collection components
-also updated the code to create new cache assets on record to not use code that may trigger UI interaction

diff --git a/Engine/Plugins/Experimental/ChaosCaching/Source/ChaosCaching/Private/Chaos/CacheManagerActor.cpp b/Engine/Plugins/Experimental/ChaosCaching/Source/ChaosCaching/Private/Chaos/CacheManagerActor.cpp
--- a/Engine/Plugins/Experimental/ChaosCaching/Source/ChaosCaching/Private/Chaos/CacheManagerActor.cpp
+++ b/Engine/Plugins/Experimental/ChaosCaching/Source/ChaosCaching/Private/Chaos/CacheManagerActor.cpp
@@ -18,6 +18,7 @@
 
 #if WITH_EDITOR
 #include "AssetToolsModule.h"
+#include "AssetRegistry/AssetRegistryModule.h"
 #include "Editor.h"
 #include "Framework/Notifications/NotificationManager.h"
 #include "UObject/ConstructorHelpers.h"
@@ -200,8 +201,18 @@ bool AChaosCacheManager::CanEditChange(const FProperty* InProperty) const
 }
 #endif
 
-
 #if WITH_EDITOR
+
+void AChaosCacheManager::PreEditChange(class FProperty* PropertyAboutToChange)
+{
+	Super::PreEditChange(PropertyAboutToChange);
+
+	if (PropertyAboutToChange && PropertyAboutToChange->GetFName() == GET_MEMBER_NAME_CHECKED(AChaosCacheManager, ObservedComponents))
+	{
+		EnsureObservedComponentsNotifierRegistration();
+	}
+}
+
 void AChaosCacheManager::PostEditChangeProperty(FPropertyChangedEvent& PropertyChangedEvent)
 {
 	Super::PostEditChangeProperty(PropertyChangedEvent);
@@ -220,9 +231,58 @@ void AChaosCacheManager::PostEditChangeProperty(FPropertyChangedEvent& PropertyC
 	{
 		SetCacheCollection(CacheCollection);
 	}
+	else if (PropertyName == GET_MEMBER_NAME_CHECKED(AChaosCacheManager, ObservedComponents))
+	{
+		EnsureObservedComponentsNotifierRegistration();
+	}
+
 }
 #endif
 
+
+void AChaosCacheManager::Destroyed()
+{
+	ClearObservedComponents(/*bDestroyCOmponents*/false);
+
+	Super::Destroyed();
+}
+
+void AChaosCacheManager::OnObservedComponentDestroyed(const FObservedComponent& ObservedComponent)
+{
+	if (UPrimitiveComponent* PrimComp = ObservedComponent.GetComponent(this))
+	{
+		PrimComp->BodyInstance.bSimulatePhysics = ObservedComponent.bIsSimulating;
+		if (UGeometryCollectionComponent* GeomComponent = Cast<UGeometryCollectionComponent>(PrimComp))
+		{
+			// in record mode we need to have notify break on on the proxy 
+			GeomComponent->SetNotifyBreaks(ObservedComponent.bHasNotifyBreaks);
+		}
+	}
+}
+
+void AChaosCacheManager::RegisterObservedComponentsNotifiers()
+{
+	for (FObservedComponent& ObservedComponent : ObservedComponents)
+	{
+		ObservedComponent.OnDestroying.BindUObject(this, &AChaosCacheManager::OnObservedComponentDestroyed);
+	}
+
+}
+
+void AChaosCacheManager::UnregisterObservedComponentsNotifiers()
+{
+	for (FObservedComponent& ObservedComponent : ObservedComponents)
+	{
+		ObservedComponent.OnDestroying.Unbind();
+	}
+}
+
+void AChaosCacheManager::EnsureObservedComponentsNotifierRegistration()
+{
+	UnregisterObservedComponentsNotifiers();
+	RegisterObservedComponentsNotifiers();
+}
+
 void AChaosCacheManager::Serialize(FArchive& Ar)
 {
 	Super::Serialize(Ar);
@@ -906,8 +966,17 @@ void AChaosCacheManager::CreateAndSetNewCacheAsset()
 		static const FString EmptyDefaultSuffix;
 		AssetToolsModule.Get().CreateUniqueAssetName(CacheCollection->GetOutermost()->GetName(), EmptyDefaultSuffix, /*out*/ NewPackageName, /*out*/ NewAssetName);
 		const FString NewPackagePath = FPackageName::GetLongPackagePath(NewPackageName);
-		UObject* NewObject = AssetToolsModule.Get().CreateAsset(NewAssetName, NewPackagePath, UChaosCacheCollection::StaticClass(), /* Factory */ nullptr);
-		CacheCollection = Cast<UChaosCacheCollection>(NewObject);
+
+		UPackage* NewPackage = CreatePackage(*NewPackageName);
+
+		EObjectFlags Flags = EObjectFlags::RF_Public | EObjectFlags::RF_Standalone;
+		UChaosCacheCollection* NewCacheCollection = NewObject<UChaosCacheCollection>(NewPackage, *NewAssetName, Flags);
+
+		NewCacheCollection->MarkPackageDirty();
+		FAssetRegistryModule::AssetCreated(NewCacheCollection);
+
+		//UObject* NewObject = AssetToolsModule.Get().CreateAsset(NewAssetName, NewPackagePath, UChaosCacheCollection::StaticClass(), /* Factory */ nullptr);
+		CacheCollection = NewCacheCollection;
 	}
 #else
 	UE_LOG(LogChaosCache, Error, TEXT("Cache Manager [%s] - Creating new asset upon recording is not supported in non-editor builds"), *GetName());
@@ -1004,6 +1073,8 @@ FObservedComponent& AChaosCacheManager::AddNewObservedComponent(UPrimitiveCompon
 		NewEntry.bHasNotifyBreaks = GeomComponent->bNotifyBreaks;
 	}
 
+	EnsureObservedComponentsNotifierRegistration();
+
 	return NewEntry;
 }
 
@@ -1020,6 +1091,10 @@ void AChaosCacheManager::FindOrAddObservedComponent(UPrimitiveComponent* InCompo
 		if(FObservedComponent* ObservedComponent = FoundComponent ? FoundComponent : &AddNewObservedComponent(InComponent))
 		{
 			ObservedComponent->bIsSimulating = bTransferSimulation ? InComponent->BodyInstance.bSimulatePhysics : false;
+			if (UGeometryCollectionComponent* GeomComponent = Cast<UGeometryCollectionComponent>(InComponent))
+			{
+				ObservedComponent->bHasNotifyBreaks = GeomComponent->bNotifyBreaks;
+			}
 
 			if(CacheName != TEXT("")) 
 			{
@@ -1031,19 +1106,27 @@ void AChaosCacheManager::FindOrAddObservedComponent(UPrimitiveComponent* InCompo
 
 void AChaosCacheManager::RemoveObservedComponent(UPrimitiveComponent* PrimitiveComponent)
 {
+	EnsureObservedComponentsNotifierRegistration();
+	// make sure we listen to those component being removed
 	ObservedComponents.RemoveAll([this, ToTest = PrimitiveComponent](const FObservedComponent& Item)
 	{
 		return Item.GetComponent(this) == ToTest;
 	});
 }
 
-void AChaosCacheManager::ClearObservedComponents()
+void AChaosCacheManager::ClearObservedComponents(bool bDestroyComponents)
 {
-	for (FObservedComponent& ObservedComponent : ObservedComponents)
+	if (bDestroyComponents)
 	{
-		UPrimitiveComponent* Comp = ObservedComponent.GetComponent(this);
-		Comp->DestroyComponent();
+		// we are destroying, no need to restore state
+		UnregisterObservedComponentsNotifiers();
+		for (FObservedComponent& ObservedComponent : ObservedComponents)
+		{
+			UPrimitiveComponent* Comp = ObservedComponent.GetComponent(this);
+			Comp->DestroyComponent();
+		}
 	}
+	EnsureObservedComponentsNotifierRegistration();
 	ObservedComponents.Empty();
 }
 
diff --git a/Engine/Plugins/Experimental/ChaosCaching/Source/ChaosCaching/Public/Chaos/CacheManagerActor.h b/Engine/Plugins/Experimental/ChaosCaching/Source/ChaosCaching/Public/Chaos/CacheManagerActor.h
--- a/Engine/Plugins/Experimental/ChaosCaching/Source/ChaosCaching/Public/Chaos/CacheManagerActor.h
+++ b/Engine/Plugins/Experimental/ChaosCaching/Source/ChaosCaching/Public/Chaos/CacheManagerActor.h
@@ -89,6 +89,13 @@ struct FObservedComponent
 		return *this;
 	}
 
+	PRAGMA_DISABLE_DEPRECATION_WARNINGS
+	~FObservedComponent()
+	{
+		OnDestroying.ExecuteIfBound(*this);
+	}
+	PRAGMA_ENABLE_DEPRECATION_WARNINGS
+
 	/** Unique name for the cache, used as a key into the cache collection */
 	UPROPERTY(EditAnywhere, BlueprintReadWrite, Category = "Caching")
 	FName CacheName;
@@ -104,7 +111,7 @@ struct FObservedComponent
 	FSoftComponentReference SoftComponentRef;
 
 	/** Capture of the initial state of the component before cache manager takes control. */
-	UPROPERTY(EditAnywhere, BlueprintReadOnly, Category = "Caching")
+	UPROPERTY(VisibleAnywhere, BlueprintReadOnly, Category = "Caching")
 	bool bIsSimulating;
 
 	/** Whether this component is enabled for playback, this allow a cache to hold many component but only replay some of them. */
@@ -120,6 +127,7 @@ struct FObservedComponent
 	* this is because when recording the cache needs the component to have bNotifyBreaks set on the component 
 	* to be able to properly record when clusters are breaking into smaller pieces
 	*/
+	UPROPERTY(VisibleAnywhere, BlueprintReadOnly, Category = "Caching")
 	bool bHasNotifyBreaks;
 
 	/** Post serialize function to transfer datas from the deprecated TObjectPtr -> TSoftObjectPtr */
@@ -153,6 +161,9 @@ private:
 	FPlaybackTickRecord TickRecord;   // Tick record to track where we are in playback
 
 	Chaos::FComponentCacheAdapter* BestFitAdapter; // Cache a pointer to the best adapter, required for cache random access handling.
+
+	DECLARE_DELEGATE_OneParam(FOnDestroying, const FObservedComponent&);
+	FOnDestroying OnDestroying;
 };
 
 struct FPerSolverData
@@ -234,6 +245,7 @@ public:
 #endif
 
 #if WITH_EDITOR
+	CHAOSCACHING_API virtual void PreEditChange(class FProperty* PropertyAboutToChange) override;
 	CHAOSCACHING_API virtual void PostEditChangeProperty(FPropertyChangedEvent& PropertyChangedEvent) override;
 #endif
 	/** end UObject interface */
@@ -319,7 +331,7 @@ public:
 	
 	/** Clear all the observed components */
 	UFUNCTION(BlueprintCallable, Category = "Caching")
-	CHAOSCACHING_API void ClearObservedComponents();
+	CHAOSCACHING_API void ClearObservedComponents(bool bDestroyComponents = true);
 	
 #if WITH_EDITOR
 	CHAOSCACHING_API void SetObservedComponentProperties(const ECacheMode& NewCacheMode);
@@ -339,6 +351,7 @@ protected:
 	/** AActor interface */
 	CHAOSCACHING_API void BeginPlay() override;
 	CHAOSCACHING_API void EndPlay(const EEndPlayReason::Type EndPlayReason) override;
+	CHAOSCACHING_API void Destroyed() override;
 	/** End AActor interface */
 	
 	/**	Handles physics thread pre-solve (push kinematic data for components under playback) */
@@ -428,6 +441,11 @@ private:
 	*/
 	void CreateAndSetNewCacheAsset();
 
+	void OnObservedComponentDestroyed(const FObservedComponent& Component);
+	void RegisterObservedComponentsNotifiers();
+	void UnregisterObservedComponentsNotifiers();
+	void EnsureObservedComponentsNotifierRegistration();
+
 	/** List of observed objects and their caches */
 	UPROPERTY(EditAnywhere, Category = "Caching")
 	TArray<FObservedComponent> ObservedComponents;
