Geometry collection cache : make static pose caches to emit events
- Events are only emitted through the Chaos Event Relay
- Only support break and collision events ( collision events do not have component or physics material as they are not stored in the cache )

diff --git a/Engine/Plugins/Experimental/ChaosCaching/Source/ChaosCaching/Private/Chaos/Adapters/GeometryCollectionComponentCacheAdapter.cpp b/Engine/Plugins/Experimental/ChaosCaching/Source/ChaosCaching/Private/Chaos/Adapters/GeometryCollectionComponentCacheAdapter.cpp
--- a/Engine/Plugins/Experimental/ChaosCaching/Source/ChaosCaching/Private/Chaos/Adapters/GeometryCollectionComponentCacheAdapter.cpp
+++ b/Engine/Plugins/Experimental/ChaosCaching/Source/ChaosCaching/Private/Chaos/Adapters/GeometryCollectionComponentCacheAdapter.cpp
@@ -8,6 +8,7 @@
 #include "GeometryCollection/GeometryCollectionComponent.h"
 #include "GeometryCollection/GeometryCollectionObject.h"
 #include "Physics/Experimental/PhysScene_Chaos.h"
+#include "Physics/Experimental/ChaosEventRelay.h"
 #include "PhysicsProxy/GeometryCollectionPhysicsProxy.h"
 #include "Chaos/ChaosSolverActor.h"
 
@@ -681,6 +682,124 @@ namespace Chaos
 		}
 	}
 
+	void FGeometryCollectionCacheAdapter::EmitEvents(UPrimitiveComponent* InComponent, UChaosCache* InCache, FPlaybackTickRecord& TickRecord) const
+	{
+		EnsureIsInGameThreadContext();
+
+		if (!InCache || InCache->GetDuration() == 0.0f)
+		{
+			return;
+		}
+
+		if (UGeometryCollectionComponent* Comp = CastChecked<UGeometryCollectionComponent>(InComponent))
+		{
+			FCacheEvaluationContext Context(TickRecord);
+			Context.bEvaluateTransform = false;
+			Context.bEvaluateCurves = false;
+			Context.bEvaluateEvents = true;
+
+			FCacheEvaluationResult EvaluatedResult = InCache->Evaluate(Context, nullptr);
+
+			EmitBreakingEvents_GameThread(*Comp, EvaluatedResult);
+			EmitCollisionEvents_GameThread(*Comp, EvaluatedResult);
+		}
+	}
+
+	void FGeometryCollectionCacheAdapter::EmitCollisionEvents_GameThread(UPrimitiveComponent& InComponent, const FCacheEvaluationResult& Results) const
+	{
+		if (const TArray<FCacheEventHandle>* CollisionEvents = Results.Events.Find(FCollisionEvent::EventName))
+		{
+			UWorld* ComponentWorld = InComponent.GetWorld();
+			if (FPhysScene* PhysicsScene = ComponentWorld ? ComponentWorld->GetPhysicsScene() : nullptr)
+			{
+				UChaosEventRelay* ChaosEventRelay = PhysicsScene->GetChaosEventRelay();
+				if (ChaosEventRelay && ChaosEventRelay->OnCollisionEvent.IsBound())
+				{
+					const FTransform ComponentToWorld = InComponent.GetComponentToWorld();
+
+					TArray<FCollisionChaosEvent> CollisionsEventsToEmit;
+					for (const FCacheEventHandle& Handle : *CollisionEvents)
+					{
+						if (const FCollisionEvent* CachedEvent = Handle.GetConst<FCollisionEvent>())
+						{
+							FCollisionChaosEvent& CollisionEvent = CollisionsEventsToEmit.AddDefaulted_GetRef();
+							CollisionEvent.Location = ComponentToWorld.TransformPosition(CachedEvent->Location);
+							CollisionEvent.AccumulatedImpulse = ComponentToWorld.TransformVectorNoScale(CachedEvent->AccumulatedImpulse);
+							CollisionEvent.Normal = ComponentToWorld.TransformVectorNoScale(CachedEvent->Normal);
+							CollisionEvent.PenetrationDepth = CachedEvent->PenetrationDepth;
+
+							// since we don't know what component each body comes from, we set both to the same
+							// this will allow at least the event handler to be able to query component related information like materials
+							CollisionEvent.Body1.Component = &InComponent; 
+							CollisionEvent.Body2.Component = &InComponent;
+
+							CollisionEvent.Body1.Velocity = ComponentToWorld.TransformVectorNoScale(CachedEvent->Velocity1);
+							CollisionEvent.Body1.DeltaVelocity = ComponentToWorld.TransformVectorNoScale(CachedEvent->DeltaVelocity1);
+							CollisionEvent.Body1.AngularVelocity = ComponentToWorld.TransformVectorNoScale(CachedEvent->AngularVelocity1);
+							CollisionEvent.Body1.Mass = CachedEvent->Mass1;
+							CollisionEvent.Body1.PhysMaterial = nullptr; // non-cached information
+							CollisionEvent.Body1.BodyIndex = INDEX_NONE; // non-cached information
+							CollisionEvent.Body1.BoneName = NAME_None; // non-cached information
+
+							CollisionEvent.Body2.Velocity = ComponentToWorld.TransformVectorNoScale(CachedEvent->Velocity2);
+							CollisionEvent.Body2.DeltaVelocity = ComponentToWorld.TransformVectorNoScale(CachedEvent->DeltaVelocity2);
+							CollisionEvent.Body2.AngularVelocity = ComponentToWorld.TransformVectorNoScale(CachedEvent->AngularVelocity2);
+							CollisionEvent.Body2.Mass = CachedEvent->Mass2;
+							CollisionEvent.Body2.PhysMaterial = nullptr; // non-cached information
+							CollisionEvent.Body2.BodyIndex = INDEX_NONE; // non-cached information
+							CollisionEvent.Body2.BoneName = NAME_None; // non-cached information
+						}
+					}
+					if (CollisionsEventsToEmit.Num() > 0)
+					{
+						ChaosEventRelay->DispatchPhysicsCollisionEvents(CollisionsEventsToEmit);
+					}
+				}
+			}
+		}
+	}
+
+	void FGeometryCollectionCacheAdapter::EmitBreakingEvents_GameThread(UPrimitiveComponent& InComponent, const FCacheEvaluationResult& Results) const
+	{
+		if (const TArray<FCacheEventHandle>* BreakingEvents = Results.Events.Find(FBreakingEvent::EventName))
+		{
+			UWorld* ComponentWorld = InComponent.GetWorld();
+			if (FPhysScene* PhysicsScene = ComponentWorld ? ComponentWorld->GetPhysicsScene() : nullptr)
+			{
+				UChaosEventRelay* ChaosEventRelay = PhysicsScene->GetChaosEventRelay();
+				if (ChaosEventRelay && ChaosEventRelay->OnBreakEvent.IsBound())
+				{
+					const FTransform ComponentToWorld = InComponent.GetComponentToWorld();
+
+					TArray<FChaosBreakEvent> BreakEventsToEmit;
+					for (const FCacheEventHandle& Handle : *BreakingEvents)
+					{
+						if (const FBreakingEvent* CachedEvent = Handle.GetConst<FBreakingEvent>())
+						{
+							FChaosBreakEvent& BreakEvent = BreakEventsToEmit.AddDefaulted_GetRef();
+							BreakEvent.Component = &InComponent;
+							BreakEvent.Location = ComponentToWorld.TransformPosition(CachedEvent->Location);
+							BreakEvent.Orientation = ComponentToWorld.TransformRotation(CachedEvent->Orientation);
+							BreakEvent.Velocity = ComponentToWorld.TransformVectorNoScale(CachedEvent->Velocity);
+							BreakEvent.AngularVelocity = ComponentToWorld.TransformVectorNoScale(CachedEvent->AngularVelocity);
+							const FVector Extents = (CachedEvent->BoundingBoxMax - CachedEvent->BoundingBoxMin) * 0.5;
+							const FBox Bounds(CachedEvent->Location - Extents, CachedEvent->Location + Extents);
+							const FBox TransformedBounds = Bounds.TransformBy(ComponentToWorld);
+							BreakEvent.Extents = TransformedBounds.GetExtent();
+							BreakEvent.Mass = CachedEvent->Mass;
+							BreakEvent.Index = CachedEvent->Index;
+							BreakEvent.bFromCrumble = false; // todo(chaos) : maybe we should cache this informnation too in the future ? 
+						}
+					}
+					if (BreakEventsToEmit.Num() > 0)
+					{
+						ChaosEventRelay->DispatchPhysicsBreakEvents(BreakEventsToEmit);
+					}
+				}
+			}
+		}
+	}
+
 	void FGeometryCollectionCacheAdapter::RegisterEventHandler(FGeometryCollectionPhysicsProxy* Proxy)
 	{
 		if (!Proxy)
diff --git a/Engine/Plugins/Experimental/ChaosCaching/Source/ChaosCaching/Private/Chaos/CacheManagerActor.cpp b/Engine/Plugins/Experimental/ChaosCaching/Source/ChaosCaching/Private/Chaos/CacheManagerActor.cpp
--- a/Engine/Plugins/Experimental/ChaosCaching/Source/ChaosCaching/Private/Chaos/CacheManagerActor.cpp
+++ b/Engine/Plugins/Experimental/ChaosCaching/Source/ChaosCaching/Private/Chaos/CacheManagerActor.cpp
@@ -79,6 +79,7 @@ AChaosCacheManager::AChaosCacheManager(const FObjectInitializer& ObjectInitializ
 	, bCanRecord(true)
 	, bIsSimulating(false)
 	, StartTimeAtBeginPlay(0)
+	, PreviousStartTime(0)
 {
 	// This actor will tick, just not normally. There needs to be a tick-like event both before physics simulation
 	// and after physics simulation, we bind to some physics scene events in BeginPlay to handle this.
@@ -207,9 +208,16 @@ void AChaosCacheManager::PreEditChange(class FProperty* PropertyAboutToChange)
 {
 	Super::PreEditChange(PropertyAboutToChange);
 
-	if (PropertyAboutToChange && PropertyAboutToChange->GetFName() == GET_MEMBER_NAME_CHECKED(AChaosCacheManager, ObservedComponents))
+	if (PropertyAboutToChange)
 	{
-		EnsureObservedComponentsNotifierRegistration();
+		if (PropertyAboutToChange->GetFName() == GET_MEMBER_NAME_CHECKED(AChaosCacheManager, ObservedComponents))
+		{
+			EnsureObservedComponentsNotifierRegistration();
+		}
+		else if (PropertyAboutToChange->GetFName() == GET_MEMBER_NAME_CHECKED(AChaosCacheManager, StartTime))
+		{
+			PreviousStartTime = StartTime;
+		}
 	}
 }
 
@@ -225,7 +233,11 @@ void AChaosCacheManager::PostEditChangeProperty(FPropertyChangedEvent& PropertyC
 	}
 	else if (PropertyName == GET_MEMBER_NAME_CHECKED(AChaosCacheManager, StartTime)) 
 	{
-		SetStartTime(StartTime);
+		// note : we cannot call SetStart time because it would override PreviousStartTime ( see also PreEditChange where PreviousStartTime is updated when the property changes )
+		if (!IsRecording())
+		{
+			OnStartFrameChanged(StartTime);
+		}
 	}
 	else if (PropertyName == GET_MEMBER_NAME_CHECKED(AChaosCacheManager, CacheCollection))
 	{
@@ -323,6 +335,7 @@ bool AChaosCacheManager::IsRecording() const
 
 void AChaosCacheManager::SetStartTime(float InStartTime)
 {
+	PreviousStartTime = StartTime;
 	StartTime = InStartTime;
 	if (!IsRecording())
 	{
@@ -416,6 +429,7 @@ void AChaosCacheManager::Start(float InStartTime)
 	if (!IsRecording())
 	{
 		StartTime = InStartTime; // todo(ccaillaud) - is that needed in record ?
+		PreviousStartTime = InStartTime;
 		OnStartFrameChanged(InStartTime);
 	}
 	BeginEvaluate();
@@ -872,6 +886,14 @@ void AChaosCacheManager::OnStartFrameChanged(Chaos::FReal InTime)
 					// @todo maybe rename SetRestState to LoadCacheAtTime
 					Observed.BestFitAdapter->InitializeForLoad(Comp, Observed);
 					Observed.BestFitAdapter->SetRestState(Comp, Observed.Cache, GetTransform(), InTime);
+
+					if (bReplayEventsInStaticPoseMode && CacheMode == ECacheMode::None)
+					{
+						FPlaybackTickRecord TickRecord;
+						TickRecord.SetLastTime(InTime);
+						TickRecord.SetDt(FMath::Max(0, InTime - PreviousStartTime));
+						Observed.BestFitAdapter->EmitEvents(Comp, Observed.Cache, TickRecord);
+					}
 				}
 				if (Token.IsOpen())
 				{
diff --git a/Engine/Plugins/Experimental/ChaosCaching/Source/ChaosCaching/Public/Chaos/Adapters/CacheAdapter.h b/Engine/Plugins/Experimental/ChaosCaching/Source/ChaosCaching/Public/Chaos/Adapters/CacheAdapter.h
--- a/Engine/Plugins/Experimental/ChaosCaching/Source/ChaosCaching/Public/Chaos/Adapters/CacheAdapter.h
+++ b/Engine/Plugins/Experimental/ChaosCaching/Source/ChaosCaching/Public/Chaos/Adapters/CacheAdapter.h
@@ -133,6 +133,11 @@ namespace Chaos
 		*/
 		virtual void SetRestState(UPrimitiveComponent* InComponent, UChaosCache* InCache, const FTransform& InRootTransform, Chaos::FReal InTime) const = 0;
 
+		/**
+		* Called from the game thread to emit events
+		*/
+		virtual void EmitEvents(UPrimitiveComponent* InComponent, UChaosCache* InCache, FPlaybackTickRecord& TickRecord) const {};
+
 		/**
 		 * Called from the game thread to initialize a component and cache ready to load a cache at a given frame
 		 */
diff --git a/Engine/Plugins/Experimental/ChaosCaching/Source/ChaosCaching/Public/Chaos/Adapters/GeometryCollectionComponentCacheAdapter.h b/Engine/Plugins/Experimental/ChaosCaching/Source/ChaosCaching/Public/Chaos/Adapters/GeometryCollectionComponentCacheAdapter.h
--- a/Engine/Plugins/Experimental/ChaosCaching/Source/ChaosCaching/Public/Chaos/Adapters/GeometryCollectionComponentCacheAdapter.h
+++ b/Engine/Plugins/Experimental/ChaosCaching/Source/ChaosCaching/Public/Chaos/Adapters/GeometryCollectionComponentCacheAdapter.h
@@ -222,6 +222,7 @@ struct FTrailingEvent : public FCacheEventBase
 
 class UChaosCache;
 class UPrimitiveComponent;
+struct FCacheEvaluationResult;
 
 namespace Chaos
 {
@@ -256,6 +257,7 @@ namespace Chaos
 		virtual Chaos::FPhysicsSolverEvents* BuildEventsSolver(UPrimitiveComponent* InComponent) const override;
 		virtual bool                   ValidForPlayback(UPrimitiveComponent* InComponent, UChaosCache* InCache) const override;
 		virtual void				   SetRestState(UPrimitiveComponent* InComponent, UChaosCache* InCache, const FTransform& InRootTransform, Chaos::FReal InTime) const override;
+		virtual void				   EmitEvents(UPrimitiveComponent* InComponent, UChaosCache* InCache, FPlaybackTickRecord& TickRecord) const override;
 		virtual bool                   InitializeForRecord(UPrimitiveComponent* InComponent, UChaosCache* InCache) override;
 		virtual bool                   InitializeForPlayback(UPrimitiveComponent* InComponent, UChaosCache* InCache, float InTime) override;
 		virtual void				   FinalizeComponent(UPrimitiveComponent* InComponent, UChaosCache* InCache) override;
@@ -275,8 +277,11 @@ namespace Chaos
 		void HandleBreakingEvents(const Chaos::FBreakingEventData& Event);
 		void HandleCollisionEvents(const Chaos::FCollisionEventData& Event);
 		void HandleTrailingEvents(const Chaos::FTrailingEventData& Event);
+		
+		void EmitCollisionEvents_GameThread(UPrimitiveComponent& InComponent, const FCacheEvaluationResult& Results) const;
+		void EmitBreakingEvents_GameThread(UPrimitiveComponent& InComponent, const FCacheEvaluationResult& Results) const;
+
 
-	
 	private:
 		TArray<int32> GatherAllBreaksUpToTime(UChaosCache* InCache, float InTime) const;
 
diff --git a/Engine/Plugins/Experimental/ChaosCaching/Source/ChaosCaching/Public/Chaos/CacheManagerActor.h b/Engine/Plugins/Experimental/ChaosCaching/Source/ChaosCaching/Public/Chaos/CacheManagerActor.h
--- a/Engine/Plugins/Experimental/ChaosCaching/Source/ChaosCaching/Public/Chaos/CacheManagerActor.h
+++ b/Engine/Plugins/Experimental/ChaosCaching/Source/ChaosCaching/Public/Chaos/CacheManagerActor.h
@@ -224,11 +224,19 @@ public:
 	UPROPERTY(EditAnywhere, Interp, BlueprintReadWrite, Category = "Caching", meta=(SequencerTrackClass="/Script/MovieSceneTracks.MovieSceneFloatTrack"))
 	float StartTime;
 
-	// if true, cache will start automatically on begin play ( playing or recording )
-	// if false, the function Start must be used to tsrat the cache again ( and can be stopped again using Stop )
+	/**
+	* if true, cache will start automatically on begin play(playing or recording)
+	* if false, the function Start must be used to tsrat the cache again ( and can be stopped again using Stop )
+	*/
 	UPROPERTY(EditAnywhere, Category = "Caching")
 	bool bStartOnBeginPlay;
 
+	/**
+	* When in static pose, replay the breaking and collision events through the Chaos Event Relay
+	*/
+	UPROPERTY(EditAnywhere, Category = "Caching", meta = (EditCondition = "CacheMode == ECacheMode::None"))
+	bool bReplayEventsInStaticPoseMode = true;
+
 	/** AActor interface */
 	CHAOSCACHING_API void TickActor(float DeltaTime, enum ELevelTick TickType, FActorTickFunction& ThisTickFunction) override;
 	/** end AActor interface */
@@ -463,6 +471,9 @@ private:
 	// since the object persists
 	float StartTimeAtBeginPlay;
 
+	// previous start time so we can keep track of the elasped time between SetStartTime calls
+	float PreviousStartTime;
+
 	bool bRestartSimulation = false; // if restart for a specified time range
 	float RestartTimeStart = TNumericLimits<float>::Max(); // if bRestartSimulation, use RestartTimeStart as an offset when saving, and clear cache data in [RestartTimeStart, RestartTimeEnd] 
 	float RestartTimeEnd = -TNumericLimits<float>::Max();
