Chaos cache manager : expose trigger method publicly and export Start / Stop / SetCurrentTime ones

#Synced-CL 49892100

diff --git a/Engine/Plugins/Experimental/ChaosCaching/Source/ChaosCaching/Public/Chaos/CacheManagerActor.h b/Engine/Plugins/Experimental/ChaosCaching/Source/ChaosCaching/Public/Chaos/CacheManagerActor.h
--- a/Engine/Plugins/Experimental/ChaosCaching/Source/ChaosCaching/Public/Chaos/CacheManagerActor.h
+++ b/Engine/Plugins/Experimental/ChaosCaching/Source/ChaosCaching/Public/Chaos/CacheManagerActor.h
@@ -275,7 +275,7 @@ public:
 	* Note this will be ignored for evaluating in Play mode ( for that case, make sure to use Start(time) instead as this will properly stop the physics thread evalution and restert it after setting the time )
 	*/
 	UFUNCTION(BlueprintCallable, Category = "Caching")
-	void SetCurrentTime(float CurrentTime);
+	CHAOSCACHING_API void SetCurrentTime(float CurrentTime);
 
 	/* 
 	* Start the cache evaluation at a specific time 
@@ -283,13 +283,13 @@ public:
 	* only works when the cache is in Static Pose or Play mode
 	*/
 	UFUNCTION(BlueprintCallable, Category = "Caching")
-	void Start(float InStartTime = 0);
+	CHAOSCACHING_API void Start(float InStartTime = 0);
 
 	/* 
 	* stop the cache from eveluating 
 	*/
 	UFUNCTION(BlueprintCallable, Category = "Caching")
-	void Stop();
+	CHAOSCACHING_API void Stop();
 
 	/** 
 	 * Resets all components back to the world space transform they had when the cache for them was originally recorded
@@ -384,6 +384,7 @@ protected:
 	UFUNCTION(BlueprintCallable, Category = "Caching")
 	CHAOSCACHING_API void SetCacheCollection(UChaosCacheCollection* InCacheCollection);
 
+public:
 	/**
 	 * Triggers a component to play or record.
 	 * If the cache manager has an observed component entry for InComponent and it is a triggered entry
@@ -406,6 +407,7 @@ protected:
 	UFUNCTION(BlueprintCallable, Category = "Caching")
 	CHAOSCACHING_API void TriggerAll();
 
+protected:
 	/** Enable playback for a specific component using its cache name */
 	UFUNCTION(BlueprintCallable, Category = "Caching")
 	CHAOSCACHING_API void EnablePlaybackByCache(FName InCacheName, bool bEnable);
