Chaos cache : change Geometry collection adapter to use kinematic target to set the position of the cached particles instead of setting them using XR and PQ
- This is needed to make sure the island get to know that the kinematic particle are moving and so they can properly wake up the connected dynamic objects accordingly

#Synced-CL 50024027

diff --git a/Engine/Plugins/Experimental/ChaosCaching/Source/ChaosCaching/Private/Chaos/Adapters/GeometryCollectionComponentCacheAdapter.cpp b/Engine/Plugins/Experimental/ChaosCaching/Source/ChaosCaching/Private/Chaos/Adapters/GeometryCollectionComponentCacheAdapter.cpp
--- a/Engine/Plugins/Experimental/ChaosCaching/Source/ChaosCaching/Private/Chaos/Adapters/GeometryCollectionComponentCacheAdapter.cpp
+++ b/Engine/Plugins/Experimental/ChaosCaching/Source/ChaosCaching/Private/Chaos/Adapters/GeometryCollectionComponentCacheAdapter.cpp
@@ -22,6 +22,13 @@ FName FTrailingEvent::EventName("GC_Trailing");
 namespace Chaos
 {
 
+	namespace GeometryCollectionCacheAdapter::CVars
+	{
+		bool bUseKinematicTarget = true;
+		FAutoConsoleVariableRef CVARUseKinematicTarget(TEXT("p.Chaos.Cache.GeometryCollectionPlayUseKinematicTarget"), bUseKinematicTarget,
+			TEXT("When in Play cache mode, use kinematic target to move the geometry collection rigid body instead of setting their position directly [def:true]"));
+	}
+
 	FComponentCacheAdapter::SupportType FGeometryCollectionCacheAdapter::SupportsComponentClass(UClass* InComponentClass) const
 	{
 		UClass* Desired = GetDesiredClass();
@@ -493,16 +500,24 @@ namespace Chaos
 				// to make sure the interpolation was done correctly in world space  
 				const FTransform WorldTransform{EvaluatedTransform};
 
-				Handle->SetP(WorldTransform.GetTranslation());
-				Handle->SetQ(WorldTransform.GetRotation());
-				Handle->SetX(Handle->GetP());
-				Handle->SetRf(Handle->GetQf());
-				
-				Handle->UpdateWorldSpaceState(WorldTransform, FVec3(0));
-
-				if(!Handle->Disabled())
+				if (GeometryCollectionCacheAdapter::CVars::bUseKinematicTarget)
 				{
-					Solver->GetEvolution()->DirtyParticle(*Handle);
+					FKinematicTarget Target = FKinematicTarget::MakePositionTarget(WorldTransform);
+					Solver->GetEvolution()->SetParticleKinematicTarget(Handle, Target);
+				}
+				else
+				{
+					Handle->SetP(WorldTransform.GetTranslation());
+					Handle->SetQ(WorldTransform.GetRotation());
+					Handle->SetX(Handle->GetP());
+					Handle->SetRf(Handle->GetQf());
+
+					Handle->UpdateWorldSpaceState(WorldTransform, FVec3(0));
+
+					if (!Handle->Disabled())
+					{
+						Solver->GetEvolution()->DirtyParticle(*Handle);
+					}
 				}
 
 				if(FRigidParticle* ClusterParent = Handle->ClusterIds().Id)
@@ -520,16 +535,25 @@ namespace Chaos
 							// that it exists, if it doesn't the child still gets set to the correct position.
 							FTransform ChildTransform = Handle->ChildToParent();
 							FTransform Result = ChildTransform.Inverse() * WorldTransform;
-							Parent->SetP(Result.GetTranslation());
-							Parent->SetX(Result.GetTranslation());
-							Parent->SetQ(Result.GetRotation());
-							Parent->SetR(Result.GetRotation());
-
-							Parent->UpdateWorldSpaceState(Result, FVec3(0));
 
-							if(!Parent->Disabled())
+							if (GeometryCollectionCacheAdapter::CVars::bUseKinematicTarget)
+							{
+								FKinematicTarget Target = FKinematicTarget::MakePositionTarget(Result);
+								Solver->GetEvolution()->SetParticleKinematicTarget(Parent, Target);
+							}
+							else
 							{
-								Solver->GetEvolution()->DirtyParticle(*Parent);
+								Parent->SetP(Result.GetTranslation());
+								Parent->SetX(Result.GetTranslation());
+								Parent->SetQ(Result.GetRotation());
+								Parent->SetR(Result.GetRotation());
+
+								Parent->UpdateWorldSpaceState(Result, FVec3(0));
+
+								if (!Parent->Disabled())
+								{
+									Solver->GetEvolution()->DirtyParticle(*Parent);
+								}
 							}
 						}
 					}
