Chaos cache : prevent crashes when accessing curves out of bounds

#Synced-CL 50953304

diff --git a/Engine/Plugins/Experimental/ChaosCaching/Source/ChaosCaching/Private/Chaos/ChaosCache.cpp b/Engine/Plugins/Experimental/ChaosCaching/Source/ChaosCaching/Private/Chaos/ChaosCache.cpp
--- a/Engine/Plugins/Experimental/ChaosCaching/Source/ChaosCaching/Private/Chaos/ChaosCache.cpp
+++ b/Engine/Plugins/Experimental/ChaosCaching/Source/ChaosCaching/Private/Chaos/ChaosCache.cpp
@@ -275,7 +275,10 @@ bool UChaosCache::FlushPendingFrames_MainPass(TQueue<FPendingFrameWrite, Mode>&
 						{
 							ChannelIndex = &ParticleToChannelCurve.Add(ParticleIndex, ChannelCurveToParticle.Add(ParticleIndex));
 						}
-						TargetCurve->RichCurves[*ChannelIndex].AddKey(NewData.Time, ChannelData.Value[PendingIndex]);
+						if (TargetCurve->RichCurves.IsValidIndex(*ChannelIndex))
+						{
+							TargetCurve->RichCurves[*ChannelIndex].AddKey(NewData.Time, ChannelData.Value[PendingIndex]);
+						}
 					}
 				}
 			}
@@ -748,7 +751,7 @@ FCacheEvaluationResult UChaosCache::Evaluate(const FCacheEvaluationContext& InCo
 void UChaosCache::BuildSpawnableFromComponent(const UPrimitiveComponent* InComponent, const FTransform& SpaceTransform)
 {
 	Spawnable.DuplicatedTemplate = StaticDuplicateObject(InComponent, this);
-	// duplication of teh component also copy the attach parent reference, but if we keep it this causes crashes because the garbage collection will have a dangling reference
+	// duplication of the component also copy the attach parent reference, but if we keep it this causes crashes because the garbage collection will have a dangling reference
 	// preventing the PIE level from being released properly and causing all sort of problem and crashes after that 
 	if (USceneComponent* SceneComponent = Cast<USceneComponent>(Spawnable.DuplicatedTemplate))
 	{
