Geometry Collection : fix issue when replaying caches that were recorded with a scale

#Synced-CL 51104518

diff --git a/Engine/Plugins/Experimental/ChaosCaching/Source/ChaosCaching/Private/Chaos/Adapters/GeometryCollectionComponentCacheAdapter.cpp b/Engine/Plugins/Experimental/ChaosCaching/Source/ChaosCaching/Private/Chaos/Adapters/GeometryCollectionComponentCacheAdapter.cpp
--- a/Engine/Plugins/Experimental/ChaosCaching/Source/ChaosCaching/Private/Chaos/Adapters/GeometryCollectionComponentCacheAdapter.cpp
+++ b/Engine/Plugins/Experimental/ChaosCaching/Source/ChaosCaching/Private/Chaos/Adapters/GeometryCollectionComponentCacheAdapter.cpp
@@ -649,7 +649,9 @@ namespace Chaos
 			const FTransform ActorToWorld = InRootTransform;
 			const FTransform WorldToComponent = Comp->GetComponentTransform().Inverse();
 			FTransform ActorToComponent = ActorToWorld * WorldToComponent;
-			
+			const FTransform CompScaleTransform = FTransform(FQuat::Identity, FVector::ZeroVector, Comp->GetComponentTransform().GetScale3D());
+			const bool bIsScaledComponent = !CompScaleTransform.GetScale3D().Equals(FVector::OneVector);
+
 			if (const UGeometryCollection* RestCollection = Comp->GetRestCollection())
 			{
 				TSharedPtr<FGeometryCollection, ESPMode::ThreadSafe> GeometryCollection = RestCollection->GetGeometryCollection();
@@ -680,6 +682,17 @@ namespace Chaos
 					RestTransforms[TransformIdx] = EvaluatedResult.Transform[CacheIdx] * ActorToComponent;
 				}
 
+				if (bIsScaledComponent)
+				{
+					const TManagedArray<FTransform>& MassToLocal = GeometryCollection->GetAttribute<FTransform>("MassToLocal", FGeometryCollection::TransformGroup);
+					for (int32 CacheIdx = 0; CacheIdx < NumCacheTransforms; ++CacheIdx)
+					{
+						const int32 TransformIdx = EvaluatedResult.ParticleIndices[CacheIdx];
+						FTransform MassToLocalTransform = MassToLocal[TransformIdx];
+						RestTransforms[TransformIdx] = MassToLocalTransform.Inverse() * CompScaleTransform * MassToLocalTransform * RestTransforms[TransformIdx];
+					}
+				}
+
 				// Set any broken clusters to identity so that they no longer influence child transforms.
 				TArray<int32> ReleaseIndices = GatherAllBreaksUpToTime(InCache, InTime);
 				const TManagedArray<int32>& Parent = GeometryCollection->Parent;
